Dataflow Cloth : fix weight map node crashing when no transfer collection has been assigned to the node

#Synced-CL 50358557

diff --git a/Engine/Plugins/ChaosClothAssetDataflowNodes/Source/ChaosClothAssetDataflowNodes/Private/ChaosClothAsset/WeightMapNode.cpp b/Engine/Plugins/ChaosClothAssetDataflowNodes/Source/ChaosClothAssetDataflowNodes/Private/ChaosClothAsset/WeightMapNode.cpp
--- a/Engine/Plugins/ChaosClothAssetDataflowNodes/Source/ChaosClothAssetDataflowNodes/Private/ChaosClothAsset/WeightMapNode.cpp
+++ b/Engine/Plugins/ChaosClothAssetDataflowNodes/Source/ChaosClothAssetDataflowNodes/Private/ChaosClothAsset/WeightMapNode.cpp
@@ -63,7 +63,13 @@ void FChaosClothAssetWeightMapNode::OnTransfer(UE::Dataflow::FContext& Context)
 	if (ClothFacade.IsValid())  // Can only act on the collection if it is a valid cloth collection
 	{
 		FManagedArrayCollection InTransferCollection = GetValue<FManagedArrayCollection>(Context, &TransferCollection);
+		if (InTransferCollection.IsEmpty())
+		{
+			Context.Error(LOCTEXT("OnTransferError_TransferCollectionEmpty", "Transfer collection is empty, no source data to transfer from"), this);
+			return;
+		}
 		const TSharedRef<const FManagedArrayCollection> TransferClothCollection = MakeShared<const FManagedArrayCollection>(MoveTemp(InTransferCollection));
+
 		FCollectionClothConstFacade TransferClothFacade(TransferClothCollection);
 		const FName InInputName = GetInputName(Context);
 
@@ -103,6 +109,15 @@ void FChaosClothAssetWeightMapNode::OnTransfer(UE::Dataflow::FContext& Context)
 
 				SetVertexWeights(ClothFacade.GetWeightMap(InInputName), RemappedWeights);
 			}
+			else
+			{
+				Context.Error(
+					FText::Format(LOCTEXT("OnTransferError_InvalidSimWeightmap", "Invalid transfer collection sim weightmap [{0}], no source data to transfer from"),
+						FText::FromName(InInputName)
+					),
+					this
+				);
+			}
 		}
 		else 
 		{
@@ -127,6 +142,15 @@ void FChaosClothAssetWeightMapNode::OnTransfer(UE::Dataflow::FContext& Context)
 
 				SetVertexWeights(ClothFacade.GetUserDefinedAttribute<float>(InInputName, ClothCollectionGroup::RenderVertices), RemappedWeights);
 			}
+			else
+			{
+				Context.Error(
+					FText::Format(LOCTEXT("InvalidRenderWeightmap", "Invalid transfer collection render weightmap [{0}], no source data to transfer from"),
+						FText::FromName(InInputName)
+					),
+					this
+				);
+			}
 		}
 	}
 }
