Add possibility to pass a selection to remesh a cloth collection sin or render mesh
- also add a way to use the weightmap as a selection by filtering values above a certain threshhold

#Synced-CL 50641382

diff --git a/Engine/Plugins/ChaosClothAssetDataflowNodes/Source/ChaosClothAssetDataflowNodes/Private/ChaosClothAsset/RemeshNode.cpp b/Engine/Plugins/ChaosClothAssetDataflowNodes/Source/ChaosClothAssetDataflowNodes/Private/ChaosClothAsset/RemeshNode.cpp
--- a/Engine/Plugins/ChaosClothAssetDataflowNodes/Source/ChaosClothAssetDataflowNodes/Private/ChaosClothAsset/RemeshNode.cpp
+++ b/Engine/Plugins/ChaosClothAssetDataflowNodes/Source/ChaosClothAssetDataflowNodes/Private/ChaosClothAsset/RemeshNode.cpp
@@ -2,31 +2,33 @@
 
 #include "ChaosClothAsset/RemeshNode.h"
 
-#include "ChaosClothAsset/ClothPatternToDynamicMesh.h"
+#include "Algo/Find.h"
+#include "Algo/RemoveIf.h"
 #include "Chaos/CollectionEmbeddedSpringConstraintFacade.h"
-#include "ChaosClothAsset/CollectionClothFacade.h"
-#include "ChaosClothAsset/CollectionClothSelectionFacade.h"
+#include "Chaos/CollectionPropertyFacade.h"
 #include "ChaosClothAsset/ClothCollectionGroup.h"
 #include "ChaosClothAsset/ClothCollectionAttribute.h"
-#include "ChaosClothAsset/ClothGeometryTools.h"
-#include "ChaosClothAsset/ClothEngineTools.h"
 #include "ChaosClothAsset/ClothDataflowTools.h"
+#include "ChaosClothAsset/ClothEngineTools.h"
+#include "ChaosClothAsset/ClothGeometryTools.h"
+#include "ChaosClothAsset/ClothPatternToDynamicMesh.h"
+#include "ChaosClothAsset/CollectionClothFacade.h"
+#include "ChaosClothAsset/CollectionClothSelectionFacade.h"
 #include "ChaosClothAsset/SimulationSelfCollisionSpheresConfigNode.h"
 #include "ChaosClothAsset/SimulationLongRangeAttachmentConfigNode.h"
 #include "ChaosClothAsset/SimulationMaxDistanceConfigNode.h"
+#include "CleaningOps/RemeshMeshOp.h"
+#include "CleaningOps/SimplifyMeshOp.h"
 #include "Dataflow/DataflowInputOutput.h"
 #include "DynamicMesh/MeshTangents.h"
 #include "DynamicMesh/MeshNormals.h"
 #include "DynamicMesh/DynamicVertexSkinWeightsAttribute.h"
-#include "CleaningOps/RemeshMeshOp.h"
-#include "CleaningOps/SimplifyMeshOp.h"
-#include "MeshUVChannelInfo.h"
-#include "MeshBoundaryLoops.h"
-#include "Chaos/CollectionPropertyFacade.h"
-#include "Algo/Find.h"
 #include "IMeshReductionManagerModule.h"
-#include "Algo/RemoveIf.h"
+#include "MeshBoundaryLoops.h"
+#include "MeshConstraintsUtil.h"
+#include "MeshUVChannelInfo.h"
 #include "Spatial/PointSetHashTable.h"
+#include "SubRegionRemesher.h"
 
 #include UE_INLINE_GENERATED_CPP_BY_NAME(RemeshNode)
 
@@ -393,10 +395,30 @@ namespace UE::Chaos::ClothAsset::Private
 	}
 
 
+	static bool EdgeIsInSelectedTriangleSet(const UE::Geometry::FDynamicMesh3& Mesh, int32 EdgeIndex, const TSet<int32>& SelectedTriangles)
+	{
+		if (!Mesh.IsEdge(EdgeIndex))
+		{
+			return false;
+		}
+
+		const UE::Geometry::FIndex2i IncidentTriangles = Mesh.GetEdgeT(EdgeIndex);
+		for (const int32 Tri : {IncidentTriangles[0], IncidentTriangles[1]})
+		{
+			if (Mesh.IsTriangle(Tri) && SelectedTriangles.Contains(Tri))
+			{
+				return true;
+			}
+		}
+
+		return false;
+	};
+
 	static void RemeshSeams(UE::Geometry::FDynamicMesh3& Mesh, TArray<TArray<FIntVector2>>& Seams, 
 		double TargetEdgeLength, 
 		const UE::Geometry::FDynamicMeshWeightAttribute* const DensityMapLayer,
-		const FVector2f& DensityLowHigh)
+		const FVector2f& DensityLowHigh,
+		const TSet<int32>* SelectedTriangles)
 	{
 		using namespace UE::Geometry;
 
@@ -404,6 +426,8 @@ namespace UE::Chaos::ClothAsset::Private
 		const double MinLength = 0.66 * TargetEdgeLength;
 		const double MaxLength = 1.33 * TargetEdgeLength;
 
+		const bool bUsingSelectedTriangles = SelectedTriangles && !SelectedTriangles->IsEmpty();
+
 		for (int32 SeamID = 0; SeamID < Seams.Num(); ++SeamID)
 		{
 			TArray<FIntVector2>& Seam = Seams[SeamID];
@@ -413,7 +437,13 @@ namespace UE::Chaos::ClothAsset::Private
 				const int32 SideAVertexA = Seam[StitchID][0];
 				const int32 SideAVertexB = Seam[StitchID + 1][0];
 
-				if (Mesh.FindEdge(SideAVertexA, SideAVertexB) == UE::Geometry::FDynamicMesh3::InvalidID)
+				const int32 EdgeSideA = Mesh.FindEdge(SideAVertexA, SideAVertexB);
+				if (EdgeSideA == UE::Geometry::FDynamicMesh3::InvalidID)
+				{
+					continue;
+				}
+
+				if (bUsingSelectedTriangles && !EdgeIsInSelectedTriangleSet(Mesh, EdgeSideA, *SelectedTriangles))
 				{
 					continue;
 				}
@@ -441,7 +471,13 @@ namespace UE::Chaos::ClothAsset::Private
 				const int32 SideBVertexA = Seam[StitchID][1];
 				const int32 SideBVertexB = Seam[StitchID + 1][1];
 
-				if (Mesh.FindEdge(SideBVertexA, SideBVertexB) == UE::Geometry::FDynamicMesh3::InvalidID)
+				const int32 EdgeSideB = Mesh.FindEdge(SideBVertexA, SideBVertexB);
+				if (EdgeSideB == UE::Geometry::FDynamicMesh3::InvalidID)
+				{
+					continue;
+				}
+
+				if (bUsingSelectedTriangles && !EdgeIsInSelectedTriangleSet(Mesh, EdgeSideB, *SelectedTriangles))
 				{
 					continue;
 				}
@@ -462,7 +498,7 @@ namespace UE::Chaos::ClothAsset::Private
 				}
 				else
 				{
-					EdgeBLength = EdgeScalingFunction(DensityLowHigh[0], DensityLowHigh[0]) * EdgeALength;
+					EdgeBLength = EdgeScalingFunction(DensityLowHigh[0], DensityLowHigh[0]) * EdgeBLength;
 				}
 
 				if ((EdgeALength < MinLength) && (EdgeBLength < MinLength))
@@ -665,7 +701,8 @@ namespace UE::Chaos::ClothAsset::Private
 		const TArray<TArray<FIntVector2>>& Seams, 
 		double TargetEdgeLength, 
 		const UE::Geometry::FDynamicMeshWeightAttribute* const DensityMapLayer,
-		const FVector2f& DensityLowHigh)
+		const FVector2f& DensityLowHigh, 
+		const TSet<int32>* SelectedTriangles)
 	{
 		using namespace UE::Geometry;
 
@@ -700,6 +737,8 @@ namespace UE::Chaos::ClothAsset::Private
 		const double MinLength = 0.66 * TargetEdgeLength;
 		const double MaxLength = 1.33 * TargetEdgeLength;
 
+		const bool bUsingSelectedTriangles = SelectedTriangles &&!SelectedTriangles->IsEmpty();
+
 		// Get the set of boundary edges up front and then check if they get invalidated later. 
 		// If we process edges inside this loop it tends to collapse a bunch of sequential edges, which can lead to very high-valence vertices
 		TArray<FIndex2i> BoundaryEdges;
@@ -721,6 +760,11 @@ namespace UE::Chaos::ClothAsset::Private
 				continue;
 			}
 
+			if (bUsingSelectedTriangles && !EdgeIsInSelectedTriangleSet(Mesh, EdgeID, *SelectedTriangles))
+			{
+				continue;
+			}
+
 			double EdgeLength = (Mesh.GetVertex(EdgeVerts[0]) - Mesh.GetVertex(EdgeVerts[1])).Length();
 
 			if (DensityMapLayer)
@@ -759,20 +803,75 @@ namespace UE::Chaos::ClothAsset::Private
 		};
 	}
 
+	/**
+	* Custom subclass of FRestrictedSubRegionRemesher that will update its own constraints. Typically users of FRestrictedSubRegionRemesher will interleave
+	* calls to BasicRemeshPass() with constraint updates, but this class will handle it internally.
+	*/
+	class FClothSubRegionRemesher : public UE::Geometry::FRestrictedSubRegionRemesher
+	{
+	public:
+		FClothSubRegionRemesher(UE::Geometry::FDynamicMesh3* Mesh, const TSet<int>& InTriangleROI) :
+			UE::Geometry::FRestrictedSubRegionRemesher(Mesh, InTriangleROI)
+		{
+		}
+	private:
+		virtual void BasicRemeshPass() override
+		{
+			UpdateROI();
+			// Reconstrain ROI boundaries
+			UE::Geometry::FMeshConstraints EditConstraints;
+			if (GetConstraints())
+			{
+				EditConstraints = *GetConstraints();
+			}
+			constexpr bool bAllowSplits = false;
+			constexpr bool bAllowSmoothing = false;
+			UE::Geometry::FMeshConstraintsUtil::ConstrainROIBoundariesInEdgeROI(EditConstraints, *GetMesh(), GetCurrentEdgeROI(), GetCurrentTriangleROI(), bAllowSplits, bAllowSmoothing);
+			SetExternalConstraints(EditConstraints);
+			UE::Geometry::FRestrictedSubRegionRemesher::BasicRemeshPass();
+		}
+	};
+
+	/**
+	* Custom RemeshOp that uses FClothSubRegionRemesher
+	*/
+	class FClothRemeshOp : public UE::Geometry::FRemeshMeshOp
+	{
+	public:
+		FClothRemeshOp(const TSet<int32>& TriangleROI) :
+			TriangleROI(TriangleROI)
+		{
+		};
+		TSet<int32> TriangleROI;
+	private:
+		virtual TUniquePtr<UE::Geometry::FRemesher> CreateRemesher(ERemeshType Type, UE::Geometry::FDynamicMesh3* TargetMesh) override
+		{
+			if (!TriangleROI.IsEmpty())
+			{
+				return MakeUnique<FClothSubRegionRemesher>(TargetMesh, TriangleROI);
+			}
+			return UE::Geometry::FRemeshMeshOp::CreateRemesher(Type, TargetMesh);
+		}
+	};
 
 	//
 	// Remeshing away from Seams/Boundaries
 	// 
+	struct FRemeshParams
+	{
+		double TargetEdgeLength = 100;
+		int32 Iterations = 10;
+		float SmoothingRate = 0.25;
+		bool bUniformSmoothing = true;
+		const TSet<int32>* FaceSelection = nullptr;
+		FString DensityMapName;
+		FVector2f DensityLowHigh = FVector2f(0,1);
+		UE::Geometry::FCompactMaps* CompactMaps = nullptr;
+	};
 
 	static bool Remesh(UE::Geometry::FDynamicMesh3& Mesh, 
-		double TargetEdgeLength, 
-		int32 Iterations, 
-		float SmoothingRate, 
-		bool bUniformSmoothing, 
-		const TArray<TArray<FIntVector2>>& Seams, 
-		const FString& DensityMapName, 
-		const FVector2f& DensityLowHigh,
-		UE::Geometry::FCompactMaps* CompactMaps = nullptr)
+		const FRemeshParams& Params,
+		const TArray<TArray<FIntVector2>>& Seams)
 	{
 		using namespace UE::Geometry;
 
@@ -792,7 +891,7 @@ namespace UE::Chaos::ClothAsset::Private
 		constexpr bool bAutoCompact = true;
 		constexpr bool bCoarsenBoundaries = false;
 
-		const ERemeshSmoothingType SmoothingType = bUniformSmoothing ? ERemeshSmoothingType::Uniform : ERemeshSmoothingType::MeanValue;
+		const ERemeshSmoothingType SmoothingType = Params.bUniformSmoothing ? ERemeshSmoothingType::Uniform : ERemeshSmoothingType::MeanValue;
 
 		// Mesh seam behavior
 		// Here we are talking UV, Normal, and Color seams, not Cloth Seams
@@ -815,7 +914,8 @@ namespace UE::Chaos::ClothAsset::Private
 		// Seam "corners" are held fixed. We use this angle threshold to determine what constitutes a seam corner
 		constexpr float SeamCornerThresholdAngleDegrees = 45.0;
 
-		FRemeshMeshOp RemeshOp;
+		TSet<int32> EmptySelection;
+		FClothRemeshOp RemeshOp(Params.FaceSelection? *Params.FaceSelection: EmptySelection);
 
 		TSharedPtr<FDynamicMesh3> SourceMesh = MakeShared<FDynamicMesh3>(MoveTemp(Mesh));
 		TSharedPtr<FDynamicMeshAABBTree3> SourceSpatial;
@@ -830,14 +930,14 @@ namespace UE::Chaos::ClothAsset::Private
 
 		RemeshOp.bDiscardAttributes = bDiscardAttributes;
 		RemeshOp.RemeshType = (bUseFullRemeshPasses) ? ERemeshType::FullPass : ERemeshType::Standard;
-		RemeshOp.RemeshIterations = Iterations;
-		RemeshOp.MaxRemeshIterations = Iterations;
+		RemeshOp.RemeshIterations = Params.Iterations;
+		RemeshOp.MaxRemeshIterations = Params.Iterations;
 		RemeshOp.ExtraProjectionIterations = 0;		// unused for regular remeshing
 		RemeshOp.TriangleCountHint = 0;				// unused for regular remeshing
-		RemeshOp.SmoothingStrength = FMath::Clamp(SmoothingRate, 0.0f, 1.0f);
+		RemeshOp.SmoothingStrength = FMath::Clamp(Params.SmoothingRate, 0.0f, 1.0f);
 		RemeshOp.SmoothingType = SmoothingType;
 
-		RemeshOp.TargetEdgeLength = TargetEdgeLength;
+		RemeshOp.TargetEdgeLength = Params.TargetEdgeLength;
 		RemeshOp.bPreserveSharpEdges = bPreserveSharpEdges;
 		RemeshOp.bFlips = bAllowFlips;
 		RemeshOp.bSplits = bAllowSplits;
@@ -860,11 +960,11 @@ namespace UE::Chaos::ClothAsset::Private
 		// RemeshOp makes a copy of the mesh to operate on, so we can't just pass the FDynamicMeshWeightAttribute pointer into CustomEdgeLengthScaleF
 		int32 DensityMapLayerIndex = -1;
 		bool bFoundDensityMapLayer = false;
-		if (!DensityMapName.IsEmpty() && SourceMesh->HasAttributes())
+		if (!Params.DensityMapName.IsEmpty() && SourceMesh->HasAttributes())
 		{
 			for (int32 WeightLayerIndex = 0; WeightLayerIndex < SourceMesh->Attributes()->NumWeightLayers(); ++WeightLayerIndex)
 			{
-				if (SourceMesh->Attributes()->GetWeightLayer(WeightLayerIndex)->GetName() == DensityMapName)
+				if (SourceMesh->Attributes()->GetWeightLayer(WeightLayerIndex)->GetName() == Params.DensityMapName)
 				{
 					DensityMapLayerIndex = WeightLayerIndex;
 					bFoundDensityMapLayer = true;
@@ -874,7 +974,7 @@ namespace UE::Chaos::ClothAsset::Private
 			}
 		}
 
-		RemeshOp.CustomEdgeLengthScaleF = [bFoundDensityMapLayer, DensityMapLayerIndex, &DensityLowHigh](const FDynamicMesh3& Mesh, int VertexA, int VertexB) -> double
+		RemeshOp.CustomEdgeLengthScaleF = [bFoundDensityMapLayer, DensityMapLayerIndex, &Params](const FDynamicMesh3& Mesh, int VertexA, int VertexB) -> double
 		{
 			if (bFoundDensityMapLayer)
 			{
@@ -884,17 +984,17 @@ namespace UE::Chaos::ClothAsset::Private
 
 				float WeightValueA;
 				DensityMapLayer->GetValue(VertexA, &WeightValueA);
-				WeightValueA = DensityLowHigh[0] * (1 - WeightValueA) + DensityLowHigh[1] * WeightValueA;
+				WeightValueA = Params.DensityLowHigh[0] * (1 - WeightValueA) + Params.DensityLowHigh[1] * WeightValueA;
 
 				float WeightValueB;
 				DensityMapLayer->GetValue(VertexB, &WeightValueB);
-				WeightValueB = DensityLowHigh[0] * (1 - WeightValueB) + DensityLowHigh[1] * WeightValueB;
+				WeightValueB = Params.DensityLowHigh[0] * (1 - WeightValueB) + Params.DensityLowHigh[1] * WeightValueB;
 
 				return EdgeScalingFunction(WeightValueA, WeightValueB);
 			}
 			else
 			{
-				return EdgeScalingFunction(DensityLowHigh[0], DensityLowHigh[0]);
+				return EdgeScalingFunction(Params.DensityLowHigh[0], Params.DensityLowHigh[0]);
 			}
 		};
 
@@ -950,7 +1050,7 @@ namespace UE::Chaos::ClothAsset::Private
 		// compact the input mesh if enabled
 		if (bAutoCompact)
 		{
-			Mesh.CompactInPlace(CompactMaps);
+			Mesh.CompactInPlace(Params.CompactMaps);
 		}
 
 		return true;
@@ -1254,11 +1354,18 @@ namespace UE::Chaos::ClothAsset::Private
 		}
 	}
 
-	static void RemeshSimMesh(const TSharedRef<const FManagedArrayCollection>& ClothCollection,
-		int32 IterationsSim,
-		double SmoothingSim,
-		const FString& DensityMapName,
-		const FVector2f& DensityLowHigh,
+	struct FRemeshSimMeshParams
+	{
+		int32 Iterations = 10;
+		double Smoothing = 0.25;
+		const TSet<int32>* FaceSelection = nullptr;
+		FString DensityMapName;
+		FVector2f DensityLowHigh = { 0,1 };
+	};
+
+	static void RemeshSimMesh(
+		const TSharedRef<const FManagedArrayCollection>& ClothCollection,
+		const FRemeshSimMeshParams& Params,
 		const TSharedRef<FManagedArrayCollection>& OutClothCollection,
 		TOptional<FText>& OutWarningMsg)
 	{
@@ -1303,7 +1410,7 @@ namespace UE::Chaos::ClothAsset::Private
 		{
 			for (int32 WeightLayerIndex = 0; WeightLayerIndex < Mesh2D.Attributes()->NumWeightLayers(); ++WeightLayerIndex)
 			{
-				if (Mesh2D.Attributes()->GetWeightLayer(WeightLayerIndex)->GetName() == DensityMapName)
+				if (Mesh2D.Attributes()->GetWeightLayer(WeightLayerIndex)->GetName() == Params.DensityMapName)
 				{
 					DensityMapLayer = Mesh2D.Attributes()->GetWeightLayer(WeightLayerIndex);
 					break;
@@ -1320,23 +1427,34 @@ namespace UE::Chaos::ClothAsset::Private
 
 		const double TargetEdgeLength = TMeshQueries<FDynamicMesh3>::AverageEdgeLength(Mesh2D);
 
-		for (int32 ResampleIter = 0; ResampleIter < IterationsSim; ++ResampleIter)
+		for (int32 ResampleIter = 0; ResampleIter < Params.Iterations; ++ResampleIter)
 		{
-			UE::Chaos::ClothAsset::Private::RemeshSeams(Mesh2D, Seams, TargetEdgeLength, DensityMapLayer, DensityLowHigh);
+			UE::Chaos::ClothAsset::Private::RemeshSeams(Mesh2D, Seams, TargetEdgeLength, DensityMapLayer, Params.DensityLowHigh, Params.FaceSelection);
 		}
 
 		// Remesh boundaries
 
-		for (int32 ResampleIter = 0; ResampleIter < IterationsSim; ++ResampleIter)
+		for (int32 ResampleIter = 0; ResampleIter < Params.Iterations; ++ResampleIter)
 		{
-			UE::Chaos::ClothAsset::Private::RemeshBoundaries(Mesh2D, Seams, TargetEdgeLength, DensityMapLayer, DensityLowHigh);
+			UE::Chaos::ClothAsset::Private::RemeshBoundaries(Mesh2D, Seams, TargetEdgeLength, DensityMapLayer, Params.DensityLowHigh, Params.FaceSelection);
 		}
 
 		// Do the remeshing of the rest of the mesh
 
 		UE::Geometry::FCompactMaps CompactMaps;
 		constexpr bool bUniformSmoothing = true;
-		UE::Chaos::ClothAsset::Private::Remesh(Mesh2D, TargetEdgeLength, IterationsSim, SmoothingSim, bUniformSmoothing, Seams, DensityMapName, DensityLowHigh, &CompactMaps);
+
+		FRemeshParams RemeshParams;
+		RemeshParams.TargetEdgeLength = TargetEdgeLength;
+		RemeshParams.Iterations = Params.Iterations;
+		RemeshParams.SmoothingRate = Params.Smoothing;
+		RemeshParams.bUniformSmoothing = bUniformSmoothing;
+		RemeshParams.DensityMapName = Params.DensityMapName;
+		RemeshParams.DensityLowHigh = Params.DensityLowHigh;
+		RemeshParams.CompactMaps = &CompactMaps;
+		RemeshParams.FaceSelection = Params.FaceSelection;
+
+		UE::Chaos::ClothAsset::Private::Remesh(Mesh2D, RemeshParams, Seams);
 
 		// Update stitches
 		for (TArray<FIntVector2>& Seam : Seams)
@@ -1501,16 +1619,21 @@ namespace UE::Chaos::ClothAsset::Private
 		}
 	}
 
+	struct FRemeshRenderMeshParams
+	{
+		EChaosClothAssetRemeshMethod RemeshMethod = EChaosClothAssetRemeshMethod::Remesh;
+		bool bRemeshSeams = false;
+		int32 SeamRemeshIterations = 1;
+		int32 Iterations = 10;
+		double Smoothing = 0.25;
+		const TSet<int32>* FaceSelection = nullptr;
+		FString DensityMapName;
+		FVector2f DensityLowHigh;
+		int32 TargetPercent = 100;
+	};
 
 	static void RemeshRenderMesh(const TSharedRef<const FManagedArrayCollection>& ClothCollection,
-		EChaosClothAssetRemeshMethod RemeshMethodRender,
-		bool bRemeshRenderSeams,
-		int32 RenderSeamRemeshIterations,
-		int32 IterationsRender,
-		double SmoothingRender,
-		const FString& DensityMapName,
-		const FVector2f& DensityLowHigh,
-		int32 TargetPercentRender,
+		const FRemeshRenderMeshParams& Params,
 		const TSharedRef<FManagedArrayCollection>& OutClothCollection)
 	{
 		using namespace UE::Geometry;
@@ -1535,7 +1658,7 @@ namespace UE::Chaos::ClothAsset::Private
 
 		TArray<TArray<FIntVector2>> Seams;
 
-		if (bRemeshRenderSeams)
+		if (Params.bRemeshSeams)
 		{
 			// Create pseudo-stitches based on boundary vertex proximity. These stitches aren't going to actually weld vertices together, but they will guide boundary remeshing.
 			// The goal is to maintain a vertex pairing along boundaries in order to avoid holes opening up when the mesh deforms due to skinning.
@@ -1550,7 +1673,7 @@ namespace UE::Chaos::ClothAsset::Private
 			{
 				for (int32 WeightLayerIndex = 0; WeightLayerIndex < DynamicMesh.Attributes()->NumWeightLayers(); ++WeightLayerIndex)
 				{
-					if (DynamicMesh.Attributes()->GetWeightLayer(WeightLayerIndex)->GetName() == DensityMapName)
+					if (DynamicMesh.Attributes()->GetWeightLayer(WeightLayerIndex)->GetName() == Params.DensityMapName)
 					{
 						DensityMapLayer = DynamicMesh.Attributes()->GetWeightLayer(WeightLayerIndex);
 						break;
@@ -1558,30 +1681,41 @@ namespace UE::Chaos::ClothAsset::Private
 				}
 			}
 
-			for (int RemeshPass = 0; RemeshPass < RenderSeamRemeshIterations; ++RemeshPass)
+			for (int RemeshPass = 0; RemeshPass < Params.SeamRemeshIterations; ++RemeshPass)
 			{
-				UE::Chaos::ClothAsset::Private::RemeshSeams(DynamicMesh, Seams, TargetEdgeLength, DensityMapLayer, DensityLowHigh);
+				UE::Chaos::ClothAsset::Private::RemeshSeams(DynamicMesh, Seams, TargetEdgeLength, DensityMapLayer, Params.DensityLowHigh, Params.FaceSelection);
 			}
 
 			// Also remesh the open boundaries that are not constrained by seams
-			for (int RemeshPass = 0; RemeshPass < RenderSeamRemeshIterations; ++RemeshPass)
+			for (int RemeshPass = 0; RemeshPass < Params.SeamRemeshIterations; ++RemeshPass)
 			{
-				UE::Chaos::ClothAsset::Private::RemeshBoundaries(DynamicMesh, Seams, TargetEdgeLength, DensityMapLayer, DensityLowHigh);
+				UE::Chaos::ClothAsset::Private::RemeshBoundaries(DynamicMesh, Seams, TargetEdgeLength, DensityMapLayer, Params.DensityLowHigh, Params.FaceSelection);
 			}
 		}
 
 
 		UE::Geometry::FCompactMaps CompactMaps;
-		if (RemeshMethodRender == EChaosClothAssetRemeshMethod::Remesh)
+		if (Params.RemeshMethod == EChaosClothAssetRemeshMethod::Remesh)
 		{
-			constexpr bool bUniformSmoothing = false;	// uniform smoothing can distort the UV layer pretty badly
-			const bool bSuccess = UE::Chaos::ClothAsset::Private::Remesh(DynamicMesh, TargetEdgeLength, IterationsRender, SmoothingRender, bUniformSmoothing, Seams, DensityMapName, DensityLowHigh, &CompactMaps);
+			constexpr bool bUniformSmo1othing = false;	// uniform smoothing can distort the UV layer pretty badly
+
+			FRemeshParams RemeshParams;
+			RemeshParams.TargetEdgeLength = TargetEdgeLength;
+			RemeshParams.Iterations = Params.Iterations;
+			RemeshParams.SmoothingRate = Params.Smoothing;
+			RemeshParams.bUniformSmoothing = bUniformSmo1othing;
+			RemeshParams.DensityMapName = Params.DensityMapName;
+			RemeshParams.DensityLowHigh = Params.DensityLowHigh;
+			RemeshParams.CompactMaps = &CompactMaps;
+			RemeshParams.FaceSelection = Params.FaceSelection;
+
+			const bool bSuccess = UE::Chaos::ClothAsset::Private::Remesh(DynamicMesh, RemeshParams, Seams);
 			check(bSuccess);
 		}
 		else
 		{
-			const bool bCoarsenBoundariesDuringSimplify = !bRemeshRenderSeams;
-			const int TargetVertexCount = FMath::RoundToInt(static_cast<float>(TargetPercentRender) / 100.0f * static_cast<float>(InputMeshVertexCount));
+			const bool bCoarsenBoundariesDuringSimplify = !Params.bRemeshSeams;
+			const int TargetVertexCount = FMath::RoundToInt(static_cast<float>(Params.TargetPercent) / 100.0f * static_cast<float>(InputMeshVertexCount));
 			UE::Chaos::ClothAsset::Private::Simplify(DynamicMesh, TargetVertexCount, bCoarsenBoundariesDuringSimplify, &CompactMaps);
 		}
 
@@ -1887,12 +2021,106 @@ FChaosClothAssetRemeshNode_v2::FChaosClothAssetRemeshNode_v2(const UE::Dataflow:
 	: FDataflowNode(InParam, InGuid)
 {
 	RegisterInputConnection(&Collection);
-	RegisterOutputConnection(&Collection, &Collection);
+	RegisterInputConnection(&TriangleSelectionSim.StringValue, GET_MEMBER_NAME_CHECKED(FChaosClothAssetConnectableIStringValue, StringValue));
 	RegisterInputConnection(&DensityMapSim.WeightMap);
+
+	RegisterInputConnection(&TriangleSelectionRender.StringValue, GET_MEMBER_NAME_CHECKED(FChaosClothAssetConnectableIStringValue, StringValue));
 	RegisterInputConnection(&DensityMapRender.WeightMap);
+
+	RegisterOutputConnection(&Collection, &Collection);
 }
 
+void FChaosClothAssetRemeshNode_v2::GetSimSelectionFromContext(UE::Dataflow::FContext& Context, const TSharedRef<FManagedArrayCollection>& ClothCollection, const FString& SimDensityMap, TSet<int32>& OutSelection) const
+{
+	using namespace UE::Chaos::ClothAsset;
+
+	OutSelection.Reset();
 
+	switch (SelectionSourceSim)
+	{
+	case EChaosClothAssetRemeshSelectionSource::SelectionInput:
+	{
+		const FName SimTriangleSelectionName = FName(*GetValue(Context, &TriangleSelectionSim.StringValue));
+		if (!FClothGeometryTools::ConvertSelectionToNewGroupType(ClothCollection, SimTriangleSelectionName, ClothCollectionGroup::SimFaces, OutSelection))
+		{
+			Context.Warning(
+				FText::Format(LOCTEXT("SimSelection_NotFound", "Invalid input sim mesh selection \"{0}\" "),
+					FText::FromName(SimTriangleSelectionName))
+				, this);
+		}
+		break;
+	}
+	case EChaosClothAssetRemeshSelectionSource::DensityMapInput:
+	{
+		FCollectionClothConstFacade ClothFacade(ClothCollection);
+		const FName SimDensityMapName(SimDensityMap);
+		if (SimDensityMapName != NAME_None && ClothFacade.HasWeightMap(SimDensityMapName))
+		{
+			TConstArrayView<float> WeightMap = ClothFacade.GetWeightMap(SimDensityMapName);
+
+			const TConstArrayView<FIntVector3> SimIndices3D = ClothFacade.GetSimIndices3D();
+			// Select a face if all 3 vertices density values are below the threshold
+			for (int32 FaceIndex = 0; FaceIndex < SimIndices3D.Num(); ++FaceIndex)
+			{
+				const FIntVector3& Index = SimIndices3D[FaceIndex];
+				if (WeightMap[Index[0]] < DensityMapThresholdSim &&
+					WeightMap[Index[1]] < DensityMapThresholdSim &&
+					WeightMap[Index[2]] < DensityMapThresholdSim)
+				{
+					OutSelection.Add(FaceIndex);
+				}
+			}
+		}
+		break;
+	}
+	}
+}
+
+void FChaosClothAssetRemeshNode_v2::GetRenderSelectionFromContext(UE::Dataflow::FContext& Context, const TSharedRef<FManagedArrayCollection>& ClothCollection, const FString& RenderDensityMap, TSet<int32>& OutSelection) const
+{
+	using namespace UE::Chaos::ClothAsset;
+
+	OutSelection.Reset();
+
+	switch (SelectionSourceRender)
+	{
+	case EChaosClothAssetRemeshSelectionSource::SelectionInput:
+	{
+		const FName RenderTriangleSelectionName = FName(*GetValue(Context, &TriangleSelectionRender.StringValue));
+		if (!FClothGeometryTools::ConvertSelectionToNewGroupType(ClothCollection, RenderTriangleSelectionName, ClothCollectionGroup::RenderFaces, OutSelection))
+		{
+			Context.Warning(
+				FText::Format(LOCTEXT("SimSelection_NotFound", "Invalid input render mesh selection \"{0}\" "),
+					FText::FromName(RenderTriangleSelectionName))
+				, this);
+		}
+		break;
+	}
+	case EChaosClothAssetRemeshSelectionSource::DensityMapInput:
+	{
+		FCollectionClothConstFacade ClothFacade(ClothCollection);
+		const FName RenderDensityMapName(RenderDensityMap);
+		if (RenderDensityMapName != NAME_None && ClothFacade.HasWeightMap(RenderDensityMapName))
+		{
+			TConstArrayView<float> WeightMap = ClothFacade.GetWeightMap(RenderDensityMapName);
+
+			const TConstArrayView<FIntVector3> SimIndices3D = ClothFacade.GetRenderIndices();
+			// Select a face if all 3 vertices density values are below the threshold
+			for (int32 FaceIndex = 0; FaceIndex < SimIndices3D.Num(); ++FaceIndex)
+			{
+				const FIntVector3& Index = SimIndices3D[FaceIndex];
+				if (WeightMap[Index[0]] < DensityMapThresholdRender &&
+					WeightMap[Index[1]] < DensityMapThresholdRender &&
+					WeightMap[Index[2]] < DensityMapThresholdRender)
+				{
+					OutSelection.Add(FaceIndex);
+				}
+			}
+		}
+		break;
+	}
+	}
+}
 void FChaosClothAssetRemeshNode_v2::Evaluate(UE::Dataflow::FContext& Context, const FDataflowOutput* Out) const
 {
 	if (Out->IsA<FManagedArrayCollection>(&Collection))
@@ -1956,8 +2184,19 @@ void FChaosClothAssetRemeshNode_v2::Evaluate(UE::Dataflow::FContext& Context, co
 				else
 				{
 					Private::EmptySimSelections(OutputClothCollection);
+
+					TSet<int32> SimFaceSelection;
+					GetSimSelectionFromContext(Context, ClothCollection, SimDensityMapName, SimFaceSelection);
+
+					Private::FRemeshSimMeshParams Params;
+					Params.Iterations = IterationsSim;
+					Params.Smoothing = SmoothingSim;
+					Params.DensityMapName = SimDensityMapName;
+					Params.DensityLowHigh = FVector2f(DensityMapSim.Low / 100.0f, DensityMapSim.High / 100.0f);
+					Params.FaceSelection = &SimFaceSelection;
+
 					TOptional<FText> WarningMsg;
-					Private::RemeshSimMesh(ClothCollection, IterationsSim, SmoothingSim, SimDensityMapName, FVector2f(DensityMapSim.Low / 100.0f, DensityMapSim.High / 100.0f), OutputClothCollection, WarningMsg);
+					Private::RemeshSimMesh(ClothCollection, Params, OutputClothCollection, WarningMsg);
 					if (WarningMsg)
 					{
 						Context.Warning(*WarningMsg, this, Out);
@@ -1969,16 +2208,22 @@ void FChaosClothAssetRemeshNode_v2::Evaluate(UE::Dataflow::FContext& Context, co
 			if (bRemeshRender)
 			{
 				Private::EmptyRenderSelections(OutputClothCollection);
-				Private::RemeshRenderMesh(ClothCollection,
-					RemeshMethodRender,
-					bRemeshRenderSeams,
-					RenderSeamRemeshIterations,
-					IterationsRender,
-					SmoothingRender,
-					RenderDensityMapName, 
-					FVector2f(DensityMapRender.Low / 100.0f, DensityMapRender.High / 100.0f), 
-					TargetPercentRender,
-					OutputClothCollection);
+
+				TSet<int32> RenderFaceSelection;
+				GetRenderSelectionFromContext(Context, ClothCollection, RenderDensityMapName, RenderFaceSelection);
+
+				Private::FRemeshRenderMeshParams Params;
+				Params.RemeshMethod = RemeshMethodRender;
+				Params.bRemeshSeams = bRemeshRenderSeams;
+				Params.SeamRemeshIterations = RenderSeamRemeshIterations;
+				Params.Iterations = IterationsRender;
+				Params.Smoothing = SmoothingRender;
+				Params.DensityMapName = RenderDensityMapName;
+				Params.DensityLowHigh = FVector2f(DensityMapRender.Low / 100.0f, DensityMapRender.High / 100.0f);
+				Params.TargetPercent = TargetPercentRender;
+				Params.FaceSelection = &RenderFaceSelection;
+
+				Private::RemeshRenderMesh(ClothCollection, Params, OutputClothCollection);
 			}
 		}
 
diff --git a/Engine/Plugins/ChaosClothAssetDataflowNodes/Source/ChaosClothAssetDataflowNodes/Public/ChaosClothAsset/RemeshNode.h b/Engine/Plugins/ChaosClothAssetDataflowNodes/Source/ChaosClothAssetDataflowNodes/Public/ChaosClothAsset/RemeshNode.h
--- a/Engine/Plugins/ChaosClothAssetDataflowNodes/Source/ChaosClothAssetDataflowNodes/Public/ChaosClothAsset/RemeshNode.h
+++ b/Engine/Plugins/ChaosClothAssetDataflowNodes/Source/ChaosClothAssetDataflowNodes/Public/ChaosClothAsset/RemeshNode.h
@@ -21,6 +21,13 @@ enum class EChaosClothAssetRemeshMethod : uint8
 	Simplify
 };
 
+UENUM()
+enum class EChaosClothAssetRemeshSelectionSource : uint8
+{
+	SelectionInput  UMETA(DisplayName = "From Selection Input"),
+	DensityMapInput UMETA(DisplayName = "From DensityMap Input (With Threshold)"),
+};
+
 /** Remesh the cloth surface(s) to get the specified mesh resolution(s).
  *  NOTE: Weight Maps, Skinning Data, Self Collision Spheres, and Long Range Attachment Constraints will be reconstructed on the output mesh, however all other Selections will be removed
  */
@@ -36,12 +43,25 @@ public:
 
 private:
 
+	/** Cloth collection to remesh */
 	UPROPERTY(Meta = (Dataflowinput, DataflowOutput, DataflowPassthrough = "Collection"))
 	FManagedArrayCollection Collection;
 
+	/** True if the sim mesh needs to be remeshed */
 	UPROPERTY(EditAnywhere, Category = "Sim Mesh")
 	bool bRemeshSim = true;
 
+	/** Defines the source for selecting which part of the sim mesh is being remeshed */
+	UPROPERTY(EditAnywhere, Category = "Sim Mesh", Meta = (EditCondition = "bRemeshSim"))
+	EChaosClothAssetRemeshSelectionSource SelectionSourceSim = EChaosClothAssetRemeshSelectionSource::DensityMapInput;
+
+	/** 
+	* Optional sim mesh triangle selection to use for remeshing 
+	* the selection defines what triangles will can remeshed, the one not in the selection will remain unchanged
+	*/
+	UPROPERTY(EditAnywhere, Category = "Sim Mesh", Meta = (EditCondition = "bRemeshSim && SelectionSourceSim == EChaosClothAssetRemeshSelectionSource::SelectionInput"))
+	FChaosClothAssetConnectableIStringValue TriangleSelectionSim;
+
 	/**
 	 * Range of target mesh resolutions, as a percentage of input triangle mesh resolution. A value of 50 on all vertices should roughly halve the total number of triangles.
 	 * If a valid vertex weight map is specified, it will use vertex weights to interpolate between the Lo and Hi values. Otherwise it will use the Lo value on all vertices.
@@ -49,18 +69,40 @@ private:
 	UPROPERTY(EditAnywhere, Category = "Sim Mesh", Meta = (EditCondition = "bRemeshSim"))
 	FChaosClothAssetWeightedValueNonAnimatable DensityMapSim = { 100.f, 200.f, TEXT("DensityMapSim") };
 
+	/** 
+	* Threshold used to define how to convert the sim mesh density into a selection map 
+	* Triangles with all their vertices above or equal to the threshold will remain untouched
+	*/
+	UPROPERTY(EditAnywhere, Category = "Sim Mesh", meta = (UIMin = "0.0", UIMax = "1.0", ClampMin = "0.0", ClampMax = "1.0", EditCondition = "bRemeshSim && SelectionSourceSim == EChaosClothAssetRemeshSelectionSource::DensityMapInput"))
+	float DensityMapThresholdSim = 1.0f;
+
+	/** Number of iterations to use for remeshing the sim mesh */
 	UPROPERTY(EditAnywhere, Category = "Sim Mesh", meta = (UIMin = "0", UIMax = "20", ClampMin = "0", ClampMax = "200", EditCondition = "bRemeshSim"))
 	int32 IterationsSim = 10;
 
+	/** Smoothing factor for the sim remeshing operation */
 	UPROPERTY(EditAnywhere, Category = "Sim Mesh", meta = (UIMin = "0.0", UIMax = "1.0", ClampMin = "0.0", ClampMax = "1.0", EditCondition = "bRemeshSim"))
 	double SmoothingSim = 0.25;
 
+	/** True if the render mesh needs to be remeshed */
 	UPROPERTY(EditAnywhere, Category = "Render Mesh")
 	bool bRemeshRender = false;
 
+	/** Method to use when remeshing the render mesh */
 	UPROPERTY(EditAnywhere, Category = "Render Mesh", meta = (UIMin = "0.0", UIMax = "1.0", ClampMin = "0.0", ClampMax = "1.0", EditCondition = "bRemeshRender"))
 	EChaosClothAssetRemeshMethod RemeshMethodRender = EChaosClothAssetRemeshMethod::Remesh;
 
+	/** Defines the source for selecting which part of the Render mesh is being remeshed */
+	UPROPERTY(EditAnywhere, Category = "Render Mesh", Meta = (EditCondition = "bRemeshRender"))
+	EChaosClothAssetRemeshSelectionSource SelectionSourceRender = EChaosClothAssetRemeshSelectionSource::DensityMapInput;
+
+	/**
+	* Optional render mesh triangle selection to use for remeshing
+	* The selection defines what traingles will can remeshed, the one not in teh selection will remain unchanged
+	*/
+	UPROPERTY(EditAnywhere, Category = "Render Mesh", Meta = (EditCondition = "bRemeshRender && && SelectionSourceRender == EChaosClothAssetRemeshSelectionSource::SelectionInput"))
+	FChaosClothAssetConnectableIStringValue TriangleSelectionRender;
+
 	/**
 	 * Range of target mesh resolutions when using the Remesh method, as a percentage of input triangle mesh resolution. A value of 50 on all vertices should roughly halve the total number of triangles.
 	 * If a valid vertex weight map is specified, it will use vertex weights to interpolate between the Lo and Hi values. Otherwise it will use the Lo value on all vertices.
@@ -68,15 +110,24 @@ private:
 	UPROPERTY(EditAnywhere, Category = "Render Mesh", Meta = (EditCondition = "bRemeshRender && RemeshMethodRender == EChaosClothAssetRemeshMethod::Remesh"))
 	FChaosClothAssetWeightedValueNonAnimatable DensityMapRender = { 100.f, 200.f, TEXT("DensityMapRender") };
 
+	/**
+	* Threshold used to define how to convert the render mesh density into a selection map
+	* Triangles with all their vertices above or equal to the threshold will remain untouched
+	*/
+	UPROPERTY(EditAnywhere, Category = "Render Mesh", meta = (UIMin = "0.0", UIMax = "1.0", ClampMin = "0.0", ClampMax = "1.0", EditCondition = "bRemeshRender && SelectionSourceRender == EChaosClothAssetRemeshSelectionSource::DensityMapInput"))
+	float DensityMapThresholdRender = 1.0f;
+
 	/**
 	 * Target mesh resolution when using the Simplify method, as a percentage of input triangle mesh resolution. A value of 50 should roughly halve the total number of triangles.
 	 */
 	UPROPERTY(EditAnywhere, Category = "Render Mesh", meta = (UIMin = "1", UIMax = "200", ClampMin = "1", EditCondition = "bRemeshRender && RemeshMethodRender == EChaosClothAssetRemeshMethod::Simplify"))
 	int32 TargetPercentRender = 100;
 
+	/** number of iterations to use for remeshing the render mesh */
 	UPROPERTY(EditAnywhere, Category = "Render Mesh", meta = (UIMin = "0", UIMax = "20", ClampMin = "0", ClampMax = "100", EditCondition = "bRemeshRender && RemeshMethodRender == EChaosClothAssetRemeshMethod::Remesh"))
 	int32 IterationsRender = 10;
 
+	/** smoothing factor for the render remeshing operation */
 	UPROPERTY(EditAnywhere, Category = "Render Mesh", meta = (UIMin = "0.0", UIMax = "1.0", ClampMin = "0.0", ClampMax = "1.0", EditCondition = "bRemeshRender && RemeshMethodRender == EChaosClothAssetRemeshMethod::Remesh"))
 	double SmoothingRender = 0.25;
 
@@ -90,6 +141,8 @@ private:
 
 	virtual void Evaluate(UE::Dataflow::FContext & Context, const FDataflowOutput * Out) const override;
 
+	void GetSimSelectionFromContext(UE::Dataflow::FContext& Context, const TSharedRef<FManagedArrayCollection>& Collection, const FString& SimDensityMap,  TSet<int32>& OutSelection) const;
+	void GetRenderSelectionFromContext(UE::Dataflow::FContext& Context, const TSharedRef<FManagedArrayCollection>& Collection, const FString& RenderDensityMap, TSet<int32>& OutSelection) const;
 };
 
 
diff --git a/Engine/Plugins/Runtime/MeshModelingToolset/Source/ModelingOperators/Public/CleaningOps/RemeshMeshOp.h b/Engine/Plugins/Runtime/MeshModelingToolset/Source/ModelingOperators/Public/CleaningOps/RemeshMeshOp.h
--- a/Engine/Plugins/Runtime/MeshModelingToolset/Source/ModelingOperators/Public/CleaningOps/RemeshMeshOp.h
+++ b/Engine/Plugins/Runtime/MeshModelingToolset/Source/ModelingOperators/Public/CleaningOps/RemeshMeshOp.h
@@ -137,7 +137,7 @@ public:
 
 protected:
 
-	UE_API TUniquePtr<FRemesher> CreateRemesher(ERemeshType Type, FDynamicMesh3* TargetMesh);
+	UE_API virtual TUniquePtr<FRemesher> CreateRemesher(ERemeshType Type, FDynamicMesh3* TargetMesh);
 
 private:
 
