Add new Strip User Attributes Node which will remove unreferenced maps, sets, and face int maps from the cloth collection to reduce the collection and asset size. Also includes the option to keep specifically named maps, sets, and face int maps so they're available for blueprints or other runtime purposes.

#virtualized
#Synced-CL 51383081

diff --git a/Engine/Plugins/ChaosClothAsset/Source/ChaosClothAsset/Private/ChaosClothAsset/ClothCollection.cpp b/Engine/Plugins/ChaosClothAsset/Source/ChaosClothAsset/Private/ChaosClothAsset/ClothCollection.cpp
--- a/Engine/Plugins/ChaosClothAsset/Source/ChaosClothAsset/Private/ChaosClothAsset/ClothCollection.cpp
+++ b/Engine/Plugins/ChaosClothAsset/Source/ChaosClothAsset/Private/ChaosClothAsset/ClothCollection.cpp
@@ -18,6 +18,7 @@ namespace UE::Chaos::ClothAsset::Private
 			ClothCollectionAttribute::PhysicsAssetSoftObjectPathName,
 			ClothCollectionAttribute::SkeletalMeshSoftObjectPathName,
 			ClothCollectionAttribute::ReferenceBoneName,
+			ClothCollectionAttribute::CopyAllUserAttributesToSimModel,
 		};
 		return LodsGroupAttributes;
 	}
@@ -338,6 +339,7 @@ namespace UE::Chaos::ClothAsset
 		PhysicsAssetSoftObjectPathName = ManagedArrayCollection->FindAttribute<FSoftObjectPath>(ClothCollectionAttribute::PhysicsAssetSoftObjectPathName, ClothCollectionGroup::Lods);
 		SkeletalMeshSoftObjectPathName = ManagedArrayCollection->FindAttribute<FSoftObjectPath>(ClothCollectionAttribute::SkeletalMeshSoftObjectPathName, ClothCollectionGroup::Lods);
 		ReferenceBoneName = ManagedArrayCollection->FindAttribute<FName>(ClothCollectionAttribute::ReferenceBoneName, ClothCollectionGroup::Lods);
+		CopyAllUserAttributesToSimModel = ManagedArrayCollection->FindAttribute<bool>(ClothCollectionAttribute::CopyAllUserAttributesToSimModel, ClothCollectionGroup::Lods);
 
 		// Solvers Group
 		SolverGravity = ManagedArrayCollection->FindAttribute<FVector3f>(ClothCollectionAttribute::SolverGravity, ClothCollectionGroup::Solvers);
@@ -479,6 +481,7 @@ namespace UE::Chaos::ClothAsset
 			PhysicsAssetSoftObjectPathName &&
 			SkeletalMeshSoftObjectPathName &&
 			ReferenceBoneName &&
+			CopyAllUserAttributesToSimModel &&
 
 			// Solvers Group
 			(!EnumHasAnyFlags(OptionalSchemas, EClothCollectionExtendedSchemas::Solvers) ||
@@ -646,6 +649,7 @@ namespace UE::Chaos::ClothAsset
 		PhysicsAssetSoftObjectPathName = &GetManagedArrayCollection()->AddAttribute<FSoftObjectPath>(ClothCollectionAttribute::PhysicsAssetSoftObjectPathName, ClothCollectionGroup::Lods);
 		SkeletalMeshSoftObjectPathName = &GetManagedArrayCollection()->AddAttribute<FSoftObjectPath>(ClothCollectionAttribute::SkeletalMeshSoftObjectPathName, ClothCollectionGroup::Lods);
 		ReferenceBoneName = &GetManagedArrayCollection()->AddAttribute<FName>(ClothCollectionAttribute::ReferenceBoneName, ClothCollectionGroup::Lods);
+		CopyAllUserAttributesToSimModel = &GetManagedArrayCollection()->AddAttribute<bool>(ClothCollectionAttribute::CopyAllUserAttributesToSimModel, ClothCollectionGroup::Lods);
 
 		// Solvers Group
 		if (EnumHasAnyFlags(OptionalSchemas, EClothCollectionExtendedSchemas::Solvers))
@@ -832,6 +836,12 @@ namespace UE::Chaos::ClothAsset
 			{
 				ReferenceBoneName = &GetManagedArrayCollection()->AddAttribute<FName>(ClothCollectionAttribute::ReferenceBoneName, ClothCollectionGroup::Lods);
 			}
+
+			// Add CopyAllUserAttributesToSimModel attribute if other LODs group attributes are defined.
+			if (PhysicsAssetSoftObjectPathName && SkeletalMeshSoftObjectPathName && ReferenceBoneName && !CopyAllUserAttributesToSimModel)
+			{
+				CopyAllUserAttributesToSimModel = &GetManagedArrayCollection()->AddAttribute<bool>(ClothCollectionAttribute::CopyAllUserAttributesToSimModel, ClothCollectionGroup::Lods);
+			}
 		}
 	}
 
diff --git a/Engine/Plugins/ChaosClothAsset/Source/ChaosClothAsset/Private/ChaosClothAsset/ClothCollection.h b/Engine/Plugins/ChaosClothAsset/Source/ChaosClothAsset/Private/ChaosClothAsset/ClothCollection.h
--- a/Engine/Plugins/ChaosClothAsset/Source/ChaosClothAsset/Private/ChaosClothAsset/ClothCollection.h
+++ b/Engine/Plugins/ChaosClothAsset/Source/ChaosClothAsset/Private/ChaosClothAsset/ClothCollection.h
@@ -137,6 +137,7 @@ namespace UE::Chaos::ClothAsset
 		UE_DEPRECATED(5.7, "Use GetSkeletalMeshSoftObjectPathName instead")
 		const TManagedArray<FString>* GetSkeletalMeshPathName() const { return nullptr; }
 		const TManagedArray<FName>* GetReferenceBoneName() const { return ReferenceBoneName; }
+		const TManagedArray<bool>* GetCopyAllUserAttributesToSimModel() const { return CopyAllUserAttributesToSimModel; }
 
 		// ~ Solvers Group
 		const TManagedArray<FVector3f>* GetSolverGravity() const { return SolverGravity; }
@@ -266,6 +267,7 @@ namespace UE::Chaos::ClothAsset
 		UE_DEPRECATED(5.7, "Use SkeletalMeshSoftObjectPathName instead")
 		const TManagedArray<FString>* SkeletalMeshPathName;
 		const TManagedArray<FName>* ReferenceBoneName;
+		const TManagedArray<bool>* CopyAllUserAttributesToSimModel;
 
 		//~ Solvers group
 		const TManagedArray<FVector3f>* SolverGravity;
@@ -427,6 +429,7 @@ namespace UE::Chaos::ClothAsset
 		UE_DEPRECATED(5.7, "Use GetSkeletalMeshSoftObjectPathName instead")
 		TManagedArray<FString>* GetSkeletalMeshPathName() { return nullptr; }
 		TManagedArray<FName>* GetReferenceBoneName() { return  const_cast<TManagedArray<FName>*>(ReferenceBoneName); }
+		TManagedArray<bool>* GetCopyAllUserAttributesToSimModel() { return const_cast<TManagedArray<bool>*>(CopyAllUserAttributesToSimModel); }
 
 		//~ Solver Group
 		TManagedArray<FVector3f>* GetSolverGravity() { return const_cast<TManagedArray<FVector3f>*>(SolverGravity); }
diff --git a/Engine/Plugins/ChaosClothAsset/Source/ChaosClothAsset/Private/ChaosClothAsset/ClothGeometryTools.cpp b/Engine/Plugins/ChaosClothAsset/Source/ChaosClothAsset/Private/ChaosClothAsset/ClothGeometryTools.cpp
--- a/Engine/Plugins/ChaosClothAsset/Source/ChaosClothAsset/Private/ChaosClothAsset/ClothGeometryTools.cpp
+++ b/Engine/Plugins/ChaosClothAsset/Source/ChaosClothAsset/Private/ChaosClothAsset/ClothGeometryTools.cpp
@@ -5,6 +5,7 @@
 #include "ChaosClothAsset/CollectionClothFacade.h"
 #include "ChaosClothAsset/CollectionClothSelectionFacade.h"
 #include "Chaos/CollectionEmbeddedSpringConstraintFacade.h"
+#include "Chaos/CollectionPropertyFacade.h"
 #include "Containers/Queue.h"
 #include "DynamicMesh/DynamicMesh3.h"
 #include "DynamicMesh/DynamicMeshAttributeSet.h"
@@ -1293,6 +1294,77 @@ namespace UE::Chaos::ClothAsset
 		Cloth.CompactSimMorphTargets();
 	}
 
+	void FClothGeometryTools::StripUnusedUserAttributes(const TSharedRef<FManagedArrayCollection>& ClothCollection, bool bStripUnusedSimWeightMaps, const TArray<FName>& ExtraSimWeightMaps, bool bStripUnusedSets, const TArray<FName>& ExtraSets, bool bStripUnusedFaceIntMaps, const TArray<FName>& ExtraFaceIntMaps, bool bStripUnusedRenderWeightMaps, const TArray<FName>& ExtraRenderWeightMaps)
+	{
+		const ::Chaos::Softs::FCollectionPropertyConstFacade Properties(ClothCollection);
+		FCollectionClothFacade Cloth(ClothCollection);
+
+		auto IsNameUsedInProperties = [&Properties](const FName& WeightMapName)->bool
+			{
+				for (int32 PropertyIndex = 0; PropertyIndex < Properties.Num(); ++PropertyIndex)
+				{
+					if (Properties.GetStringValue(PropertyIndex) == WeightMapName)
+					{
+						return true;
+					}
+				}
+				return false;
+			};
+
+		auto IsNameReserved = [&Properties](const FName& WeightMapName)->bool
+			{
+				return WeightMapName.ToString().StartsWith(TEXT("_"));
+			};
+
+		if (bStripUnusedSimWeightMaps)
+		{
+			const TArray<FName> WeightMapNames = Cloth.GetWeightMapNames();
+			for (const FName& WeightMapName : WeightMapNames)
+			{
+				if (!(IsNameUsedInProperties(WeightMapName) || IsNameReserved(WeightMapName) || ExtraSimWeightMaps.Contains(WeightMapName)))
+				{
+					Cloth.RemoveWeightMap(WeightMapName);
+				}
+			}
+		}
+
+		if (bStripUnusedSets)
+		{
+			FCollectionClothSelectionFacade Selection(ClothCollection);
+			const TArray<FName> SelectionNames = Selection.GetNames();
+			for (const FName& SelectionName : SelectionNames)
+			{
+				if (!(IsNameUsedInProperties(SelectionName) || IsNameReserved(SelectionName) || ExtraSets.Contains(SelectionName)))
+				{
+					Selection.RemoveSelectionSet(SelectionName);
+				}
+			}
+		}
+
+		if (bStripUnusedFaceIntMaps)
+		{
+			const TArray<FName> FaceIntMapNames = Cloth.GetUserDefinedAttributeNames<int32>(ClothCollectionGroup::SimFaces);
+			for (const FName& FaceIntMapName : FaceIntMapNames)
+			{
+				if (!(IsNameUsedInProperties(FaceIntMapName) || IsNameReserved(FaceIntMapName) || ExtraFaceIntMaps.Contains(FaceIntMapName)))
+				{
+					Cloth.RemoveUserDefinedAttribute(FaceIntMapName, ClothCollectionGroup::SimFaces);
+				}
+			}
+		}
+
+		if (bStripUnusedRenderWeightMaps)
+		{
+			const TArray<FName> WeightMapNames = Cloth.GetUserDefinedAttributeNames<float>(ClothCollectionGroup::RenderVertices);
+			for (const FName& WeightMapName : WeightMapNames)
+			{
+				if (!(IsNameUsedInProperties(WeightMapName) || IsNameReserved(WeightMapName) || ExtraRenderWeightMaps.Contains(WeightMapName)))
+				{
+					Cloth.RemoveUserDefinedAttribute(WeightMapName, ClothCollectionGroup::RenderVertices);
+				}
+			}
+		}
+	}
 
 	void FClothGeometryTools::BuildConnectedSeams(const TArray<FIntVector2>& InputStitches,
 		const UE::Geometry::FDynamicMesh3& Mesh,
diff --git a/Engine/Plugins/ChaosClothAsset/Source/ChaosClothAsset/Private/ChaosClothAsset/CollectionClothFacade.cpp b/Engine/Plugins/ChaosClothAsset/Source/ChaosClothAsset/Private/ChaosClothAsset/CollectionClothFacade.cpp
--- a/Engine/Plugins/ChaosClothAsset/Source/ChaosClothAsset/Private/ChaosClothAsset/CollectionClothFacade.cpp
+++ b/Engine/Plugins/ChaosClothAsset/Source/ChaosClothAsset/Private/ChaosClothAsset/CollectionClothFacade.cpp
@@ -63,6 +63,7 @@ namespace UE::Chaos::ClothAsset
 		ResultHash = HashCombineFast(ResultHash, GetTypeHash(GetPhysicsAssetSoftObjectPathName()));
 		ResultHash = HashCombineFast(ResultHash, GetTypeHash(GetSkeletalMeshSoftObjectPathName()));
 		ResultHash = HashCombineFast(ResultHash, GetTypeHash(GetReferenceBoneName()));
+		ResultHash = HashCombineFast(ResultHash, GetTypeHash(GetCopyAllUserAttributesToSimModel()));
 
 		//~ Solvers Group
 		ResultHash = HashCombineFast(ResultHash, GetTypeHash(GetSolverGravity()));
@@ -212,6 +213,11 @@ namespace UE::Chaos::ClothAsset
 		return ClothCollection->GetReferenceBoneName() && ClothCollection->GetNumElements(ClothCollectionGroup::Lods) > 0 ? (*ClothCollection->GetReferenceBoneName())[0] : NAME_None;
 	}
 
+	bool FCollectionClothConstFacade::GetCopyAllUserAttributesToSimModel() const
+	{
+		return ClothCollection->GetCopyAllUserAttributesToSimModel() && ClothCollection->GetNumElements(ClothCollectionGroup::Lods) > 0 ? (*ClothCollection->GetCopyAllUserAttributesToSimModel())[0] : false;
+	}
+
 	bool FCollectionClothConstFacade::HasSolverElement() const
 	{
 		return ClothCollection->GetNumElements(ClothCollectionGroup::Solvers) == 1;
@@ -774,6 +780,10 @@ namespace UE::Chaos::ClothAsset
 		{
 			SetReferenceBoneName(Other.GetReferenceBoneName());
 		}
+		if (GetCopyAllUserAttributesToSimModel() == false)
+		{
+			SetCopyAllUserAttributesToSimModel(Other.GetCopyAllUserAttributesToSimModel());
+		}
 
 		const bool bOtherHasResizingData = Other.IsValid(EClothCollectionExtendedSchemas::Resizing);
 		if (bOtherHasResizingData)
@@ -1021,6 +1031,14 @@ namespace UE::Chaos::ClothAsset
 			(*GetClothCollection()->GetReferenceBoneName())[0] = BoneName;
 		}
 	}
+
+	void FCollectionClothFacade::SetCopyAllUserAttributesToSimModel(bool bValue)
+	{
+		if (ClothCollection->GetNumElements(ClothCollectionGroup::Lods) && GetClothCollection()->GetCopyAllUserAttributesToSimModel())
+		{
+			(*GetClothCollection()->GetCopyAllUserAttributesToSimModel())[0] = bValue;
+		}
+	}
 	
 	void FCollectionClothFacade::SetSolverGravity(const FVector3f& SolverGravity)
 	{
@@ -1650,6 +1668,7 @@ namespace UE::Chaos::ClothAsset
 	void FCollectionClothFacade::SetDefaults()
 	{
 		GetClothCollection()->SetNumElements(1, ClothCollectionGroup::Lods);
+		SetCopyAllUserAttributesToSimModel(false);
 
 		if (IsValid(EClothCollectionExtendedSchemas::Solvers))
 		{
diff --git a/Engine/Plugins/ChaosClothAsset/Source/ChaosClothAsset/Public/ChaosClothAsset/ClothCollectionAttribute.h b/Engine/Plugins/ChaosClothAsset/Source/ChaosClothAsset/Public/ChaosClothAsset/ClothCollectionAttribute.h
--- a/Engine/Plugins/ChaosClothAsset/Source/ChaosClothAsset/Public/ChaosClothAsset/ClothCollectionAttribute.h
+++ b/Engine/Plugins/ChaosClothAsset/Source/ChaosClothAsset/Public/ChaosClothAsset/ClothCollectionAttribute.h
@@ -15,6 +15,7 @@ namespace UE::Chaos::ClothAsset
 		inline const FName PhysicsAssetSoftObjectPathName(TEXT("PhysicsAssetSoftObjectPathName"));
 		inline const FName SkeletalMeshSoftObjectPathName(TEXT("SkeletalMeshSoftObjectPathName"));
 		inline const FName ReferenceBoneName(TEXT("ReferenceBoneName"));
+		inline const FName CopyAllUserAttributesToSimModel(TEXT("CopyAllUserAttributesToSimModel"));
 
 		// Solvers Group
 		inline const FName SolverGravity(TEXT("SolverGravity"));
diff --git a/Engine/Plugins/ChaosClothAsset/Source/ChaosClothAsset/Public/ChaosClothAsset/ClothGeometryTools.h b/Engine/Plugins/ChaosClothAsset/Source/ChaosClothAsset/Public/ChaosClothAsset/ClothGeometryTools.h
--- a/Engine/Plugins/ChaosClothAsset/Source/ChaosClothAsset/Public/ChaosClothAsset/ClothGeometryTools.h
+++ b/Engine/Plugins/ChaosClothAsset/Source/ChaosClothAsset/Public/ChaosClothAsset/ClothGeometryTools.h
@@ -101,6 +101,9 @@ namespace UE::Chaos::ClothAsset
 		*/
 		static UE_API void CleanupAndCompactMesh(const TSharedRef<FManagedArrayCollection>& ClothCollection);
 
+		/** Strip maps, sets, etc. that are not referenced by properties from a cloth collection. */
+		static UE_API void StripUnusedUserAttributes(const TSharedRef<FManagedArrayCollection>& ClothCollection, bool bStripUnusedSimWeightMaps, const TArray<FName>& ExtraSimWeightMaps, bool bStripUnusedSets, const TArray<FName>& ExtraSets, bool bStripUnusedFaceIntMaps, const TArray<FName>& ExtraFaceIntMaps, bool bStripUnusedRenderWeightMaps, const TArray<FName>& ExtraRenderWeightMaps);
+
 		/**
 		 * Find sets of connected stitches for the input stitches given in random order.
 	     * Stitch (A, B) is connected to stitch (C, D) if there exist edges {(A, C), (B, D)} *or* {(A, D), (B, C)} in the given DynamicMesh.
diff --git a/Engine/Plugins/ChaosClothAsset/Source/ChaosClothAsset/Public/ChaosClothAsset/CollectionClothFacade.h b/Engine/Plugins/ChaosClothAsset/Source/ChaosClothAsset/Public/ChaosClothAsset/CollectionClothFacade.h
--- a/Engine/Plugins/ChaosClothAsset/Source/ChaosClothAsset/Public/ChaosClothAsset/CollectionClothFacade.h
+++ b/Engine/Plugins/ChaosClothAsset/Source/ChaosClothAsset/Public/ChaosClothAsset/CollectionClothFacade.h
@@ -87,6 +87,7 @@ namespace UE::Chaos::ClothAsset
 		}
 		/** Return the reference bone name used for this collection. */
 		UE_API FName GetReferenceBoneName() const;
+		UE_API bool GetCopyAllUserAttributesToSimModel() const;
 
 		//~ Solver (single per collection) Group
 		/** Return true if the solver group has one element*/
@@ -309,6 +310,8 @@ namespace UE::Chaos::ClothAsset
 		void SetSkeletalMeshPathName(const FString& PathName){}
 		/** Set the reference bone name for this asset.*/
 		UE_API void SetReferenceBoneName(const FName& BoneName);
+		/** Set whether all attributes should be copied to the FChaosClothSimulationModel (else legacy behavior of skipping unused sets and face int maps will be followed).*/
+		UE_API void SetCopyAllUserAttributesToSimModel(bool bValue);
 
 		//~ Solver (max 1 per collection) Group
 		/** Set the solver gravity */
diff --git a/Engine/Plugins/ChaosClothAsset/Source/ChaosClothAssetEngine/Private/ChaosClothAsset/ClothSimulationModel.cpp b/Engine/Plugins/ChaosClothAsset/Source/ChaosClothAssetEngine/Private/ChaosClothAsset/ClothSimulationModel.cpp
--- a/Engine/Plugins/ChaosClothAsset/Source/ChaosClothAssetEngine/Private/ChaosClothAsset/ClothSimulationModel.cpp
+++ b/Engine/Plugins/ChaosClothAsset/Source/ChaosClothAssetEngine/Private/ChaosClothAsset/ClothSimulationModel.cpp
@@ -219,8 +219,16 @@ FChaosClothSimulationModel::FChaosClothSimulationModel(const TArray<TSharedRef<c
 		TArray<TArray<int32>> WeldedToPatternIndices;
 		Cloth.BuildSimulationMesh(LodModel.Positions, LodModel.Normals, LodModel.Indices, LodModel.PatternPositions, LodModel.PatternIndices, LodModel.PatternToWeldedIndices, &WeldedToPatternIndices);
 
-		// Copy weight maps
 		const ::Chaos::Softs::FCollectionPropertyConstFacade Properties(ClothCollections[LodIndex]);
+		const bool bCopyAllUserAttributes = Cloth.GetCopyAllUserAttributesToSimModel();
+
+		// Copy weight maps
+		LodModel.WeightMaps.Reserve(WeightMapNames.Num());
+		for (const FName& WeightMapName : WeightMapNames)
+		{
+			// Keep all weight maps to allow runtime switching. This is legacy behavior even regardless of value of bCopyAllUserAttributes
+			LodModel.WeightMaps.Add(WeightMapName) = Cloth.GetWeightMap(WeightMapName);
+		}
 
 		auto IsNameUsedInProperties = [&Properties](const FName& WeightMapName)->bool
 			{
@@ -239,13 +247,6 @@ FChaosClothSimulationModel::FChaosClothSimulationModel(const TArray<TSharedRef<c
 				return WeightMapName.ToString().StartsWith(TEXT("_"));
 			};
 
-		LodModel.WeightMaps.Reserve(WeightMapNames.Num());
-		for (const FName& WeightMapName : WeightMapNames)
-		{
-			// Keep all weight maps to allow runtime switching.
-			LodModel.WeightMaps.Add(WeightMapName) = Cloth.GetWeightMap(WeightMapName);
-		}
-
 		// Copy vertex and face sets
 		FCollectionClothSelectionConstFacade Selection(ClothCollections[LodIndex]);
 		const TArray<FName> SelectionNames = Selection.GetNames();
@@ -253,7 +254,7 @@ FChaosClothSimulationModel::FChaosClothSimulationModel(const TArray<TSharedRef<c
 		LodModel.FaceSets.Reserve(SelectionNames.Num());
 		for (const FName& SelectionName : SelectionNames)
 		{
-			if (IsNameUsedInProperties(SelectionName) || IsNameReserved(SelectionName))
+			if (bCopyAllUserAttributes || IsNameUsedInProperties(SelectionName) || IsNameReserved(SelectionName))
 			{
 				const FName SelectionGroup = Selection.GetSelectionGroup(SelectionName);
 				if (SelectionGroup == ClothCollectionGroup::SimVertices3D)
@@ -272,7 +273,7 @@ FChaosClothSimulationModel::FChaosClothSimulationModel(const TArray<TSharedRef<c
 		LodModel.FaceIntMaps.Reserve(FaceIntMapNames.Num());
 		for (const FName& FaceIntMapName : FaceIntMapNames)
 		{
-			if (IsNameUsedInProperties(FaceIntMapName) || IsNameReserved(FaceIntMapName))
+			if (bCopyAllUserAttributes || IsNameUsedInProperties(FaceIntMapName) || IsNameReserved(FaceIntMapName))
 			{
 				LodModel.FaceIntMaps.Add(FaceIntMapName) = Cloth.GetUserDefinedAttribute<int32>(FaceIntMapName, ClothCollectionGroup::SimFaces);
 			}
diff --git a/Engine/Plugins/ChaosClothAssetDataflowNodes/Source/ChaosClothAssetDataflowNodes/Private/ChaosClothAsset/ChaosClothAssetDataflowNodesModule.cpp b/Engine/Plugins/ChaosClothAssetDataflowNodes/Source/ChaosClothAssetDataflowNodes/Private/ChaosClothAsset/ChaosClothAssetDataflowNodesModule.cpp
--- a/Engine/Plugins/ChaosClothAssetDataflowNodes/Source/ChaosClothAssetDataflowNodes/Private/ChaosClothAsset/ChaosClothAssetDataflowNodesModule.cpp
+++ b/Engine/Plugins/ChaosClothAssetDataflowNodes/Source/ChaosClothAssetDataflowNodes/Private/ChaosClothAsset/ChaosClothAssetDataflowNodesModule.cpp
@@ -72,6 +72,7 @@
 #include "ChaosClothAsset/SkeletalMeshImportNode.h"
 #include "ChaosClothAsset/SkinningBlendNode.h"
 #include "ChaosClothAsset/StaticMeshImportNode.h"
+#include "ChaosClothAsset/StripUserAttributesNode.h"
 #include "ChaosClothAsset/TerminalNode.h"
 #include "ChaosClothAsset/TransferSkinWeightsNode.h"
 #include "ChaosClothAsset/TransformPositionsNode.h"
@@ -174,6 +175,7 @@ namespace UE::Chaos::ClothAsset
 			DATAFLOW_NODE_REGISTER_CREATION_FACTORY(FChaosClothAssetSkeletalMeshImportNode_v2);
 			DATAFLOW_NODE_REGISTER_CREATION_FACTORY(FChaosClothAssetSkinningBlendNode);
 			DATAFLOW_NODE_REGISTER_CREATION_FACTORY(FChaosClothAssetStaticMeshImportNode_v2);
+			DATAFLOW_NODE_REGISTER_CREATION_FACTORY(FChaosClothStripUserAttributesNode);
 			DATAFLOW_NODE_REGISTER_CREATION_FACTORY(FChaosClothAssetTerminalNode_v2);
 			DATAFLOW_NODE_REGISTER_CREATION_FACTORY(FChaosClothAssetTransferSkinWeightsNode);
 			DATAFLOW_NODE_REGISTER_CREATION_FACTORY(FChaosClothAssetTransformPositionsNode);
diff --git a/Engine/Plugins/ChaosClothAssetDataflowNodes/Source/ChaosClothAssetDataflowNodes/Private/ChaosClothAsset/StripUserAttributesNode.cpp b/Engine/Plugins/ChaosClothAssetDataflowNodes/Source/ChaosClothAssetDataflowNodes/Private/ChaosClothAsset/StripUserAttributesNode.cpp
new file mode 100644
--- /dev/null
+++ b/Engine/Plugins/ChaosClothAssetDataflowNodes/Source/ChaosClothAssetDataflowNodes/Private/ChaosClothAsset/StripUserAttributesNode.cpp
@@ -0,0 +1,62 @@
+// Copyright Epic Games, Inc. All Rights Reserved.
+#include "ChaosClothAsset/StripUserAttributesNode.h"
+#include "ChaosClothAsset/ClothGeometryTools.h"
+#include "ChaosClothAsset/CollectionClothFacade.h"
+
+#include UE_INLINE_GENERATED_CPP_BY_NAME(StripUserAttributesNode)
+
+FChaosClothStripUserAttributesNode::FChaosClothStripUserAttributesNode(const UE::Dataflow::FNodeParameters& InParam, FGuid InGuid)
+	: FDataflowNode(InParam, InGuid)
+{
+	using namespace UE::Chaos::ClothAsset;
+	RegisterInputConnection(&Collection);
+	RegisterOutputConnection(&Collection)
+		.SetPassthroughInput(&Collection);
+	RegisterInputConnection(&ExtraSimWeightMaps)
+		.SetCanHidePin(true)
+		.SetPinIsHidden(true);
+	RegisterInputConnection(&ExtraSets)
+		.SetCanHidePin(true)
+		.SetPinIsHidden(true);
+	RegisterInputConnection(&ExtraSimFaceIntMaps)
+		.SetCanHidePin(true)
+		.SetPinIsHidden(true);
+	RegisterInputConnection(&ExtraRenderWeightMaps)
+		.SetCanHidePin(true)
+		.SetPinIsHidden(true);
+}
+void FChaosClothStripUserAttributesNode::Evaluate(UE::Dataflow::FContext& Context, const FDataflowOutput* Out) const
+{
+	if (Out->IsA<FManagedArrayCollection>(&Collection))
+	{
+		using namespace UE::Chaos::ClothAsset;
+		// Evaluate in collection
+		FManagedArrayCollection InCollection = GetValue<FManagedArrayCollection>(Context, &Collection);
+		const TSharedRef<FManagedArrayCollection> ClothCollection = MakeShared<FManagedArrayCollection>(MoveTemp(InCollection));
+		// Always check for a valid cloth collection/facade/sim mesh to avoid processing non cloth collections or pure render mesh cloth assets
+		FCollectionClothFacade ClothFacade(ClothCollection);
+		if (ClothFacade.IsValid() && ClothFacade.HasValidData())
+		{
+			auto ConvertToFNameArray = [](const TArray<FString>& InStringArray)->TArray<FName>
+				{
+					TArray<FName> OutFNameArray;
+					OutFNameArray.Reset(InStringArray.Num());
+					for (const FString& InString : InStringArray)
+					{
+						OutFNameArray.Emplace(FName(InString));
+					}
+					return OutFNameArray;
+				};
+
+			UE::Chaos::ClothAsset::FClothGeometryTools::StripUnusedUserAttributes(ClothCollection,
+				bStripUnusedSimWeightMaps, ConvertToFNameArray(GetValue<TArray<FString>>(Context, &ExtraSimWeightMaps)),
+				bStripUnusedSets, ConvertToFNameArray(GetValue<TArray<FString>>(Context, &ExtraSets)),
+				bStripUnusedSimFaceIntMaps, ConvertToFNameArray(GetValue<TArray<FString>>(Context, &ExtraSimFaceIntMaps)),
+				bStripUnusedRenderWeightMaps, ConvertToFNameArray(GetValue<TArray<FString>>(Context, &ExtraRenderWeightMaps))
+			);
+
+			ClothFacade.SetCopyAllUserAttributesToSimModel(bCopyAllUserAttributesToSimModel);
+		}
+		SetValue(Context, MoveTemp(*ClothCollection), &Collection);
+	}
+}
diff --git a/Engine/Plugins/ChaosClothAssetDataflowNodes/Source/ChaosClothAssetDataflowNodes/Public/ChaosClothAsset/StripUserAttributesNode.h b/Engine/Plugins/ChaosClothAssetDataflowNodes/Source/ChaosClothAssetDataflowNodes/Public/ChaosClothAsset/StripUserAttributesNode.h
new file mode 100644
--- /dev/null
+++ b/Engine/Plugins/ChaosClothAssetDataflowNodes/Source/ChaosClothAssetDataflowNodes/Public/ChaosClothAsset/StripUserAttributesNode.h
@@ -0,0 +1,72 @@
+// Copyright Epic Games, Inc. All Rights Reserved.
+
+#pragma once
+
+#include "Dataflow/DataflowNode.h"
+#include "GeometryCollection/ManagedArrayCollection.h"
+#include "StripUserAttributesNode.generated.h"
+
+/**
+ * Strip User Attributes such as Weight Maps, Sets, and Face Int Maps from a Cloth Collection if they are not referenced by properties.
+ * This utility node makes it easier to remove unused attributes from your asset, reducing asset size.
+ * Note that if this node is not present in your Cloth Dataflow Graph, unreferenced Sets and Face Int Maps will not be copied to the
+ * simulation model. Weight Maps will continue to be copied without this node to retain legacy behavior.
+ */
+USTRUCT(Meta = (DataflowCloth))
+struct FChaosClothStripUserAttributesNode : public FDataflowNode
+{
+	GENERATED_USTRUCT_BODY()
+	DATAFLOW_NODE_DEFINE_INTERNAL(FChaosClothStripUserAttributesNode, "StripUserAttributes", "Cloth", "Cloth Strip User Attributes")
+
+public:
+
+	FChaosClothStripUserAttributesNode(const UE::Dataflow::FNodeParameters& InParam, FGuid InGuid = FGuid::NewGuid());
+
+private:
+
+	UPROPERTY(Meta = (DataflowInput, DataflowOutput, DataflowPassthrough = "Collection"))
+	FManagedArrayCollection Collection;
+
+	/** Strip all unused sim vertex weight maps from the asset. This reduces the asset size, but the weight maps will be unavailable in blueprints, for debug display, or in downstream cloth assets which import this one. */
+	UPROPERTY(EditAnywhere, Category = "Strip User Attributes")
+	bool bStripUnusedSimWeightMaps = true;
+
+	/** Extra sim weight maps to keep when stripping unused weight maps. */
+	UPROPERTY(EditAnywhere, Category = "Strip User Attributes", Meta = (DataflowInput, EditCondition = "bStripUnusedSimWeightMaps"))
+	TArray<FString> ExtraSimWeightMaps;
+
+	/** Strip all unused sim and render sets from the asset. This reduces the asset size, but the sets will be unavailable in blueprints, for debug display, or in downstream cloth assets which import this one. */
+	UPROPERTY(EditAnywhere, Category = "Strip User Attributes")
+	bool bStripUnusedSets = true;
+
+	/** Extra sets to keep when stripping unused sets. */
+	UPROPERTY(EditAnywhere, Category = "Strip User Attributes", Meta = (DataflowInput, EditCondition = "bStripUnusedSets"))
+	TArray<FString> ExtraSets;
+
+	/** Strip all unused sim face integer maps from the asset. This reduces the asset size, but the maps will be unavailable in blueprints, for debug display, or in downstream cloth assets which import this one. */
+	UPROPERTY(EditAnywhere, Category = "Strip User Attributes")
+	bool bStripUnusedSimFaceIntMaps = true;
+
+	/** Extra maps to keep when stripping unused face integer maps. */
+	UPROPERTY(EditAnywhere, Category = "Strip User Attributes", Meta = (DataflowInput, EditCondition = "bStripUnusedSimFaceIntMaps"))
+	TArray<FString> ExtraSimFaceIntMaps;
+
+	/** Strip all unused render vertex weight maps from the asset. This reduces the uncooked asset size, but the maps will be unavailable in downstream assets which import this one. 
+	 *  Note that user render weight maps do not currently exist in the cooked asset, so this property has no effect on cooked assets. */
+	UPROPERTY(EditAnywhere, Category = "Strip User Attributes")
+	bool bStripUnusedRenderWeightMaps = true;
+
+	/** Extra render weight maps to keep when stripping unused weight maps. */
+	UPROPERTY(EditAnywhere, Category = "Strip User Attributes", Meta = (DataflowInput, EditCondition = "bStripUnusedRenderWeightMaps"))
+	TArray<FString> ExtraRenderWeightMaps;
+
+	/** Tell the cloth asset to transfer all user attributes from the cloth collection to the internal simulation model rather than skipping unreferenced sets and face int maps. 
+	* Note that sim weight maps are always transferred.
+	*/
+	UPROPERTY(EditAnywhere, Category = "Strip User Attributes")
+	bool bCopyAllUserAttributesToSimModel = true;
+
+	//~ Begin FDataflowNode interface
+	virtual void Evaluate(UE::Dataflow::FContext& Context, const FDataflowOutput* Out) const override;
+	//~ End FDataflowNode interface
+};
