Chaos Cloth Asset - Fixed the non animatable property update when the cloth assets are evaluated without a change of geometry.

#rnx

#Synced-CL 51543885

diff --git a/Engine/Plugins/ChaosClothAsset/Source/ChaosClothAssetEngine/Private/ChaosClothAsset/ClothAsset.cpp b/Engine/Plugins/ChaosClothAsset/Source/ChaosClothAssetEngine/Private/ChaosClothAsset/ClothAsset.cpp
--- a/Engine/Plugins/ChaosClothAsset/Source/ChaosClothAssetEngine/Private/ChaosClothAsset/ClothAsset.cpp
+++ b/Engine/Plugins/ChaosClothAsset/Source/ChaosClothAssetEngine/Private/ChaosClothAsset/ClothAsset.cpp
@@ -507,7 +507,7 @@ void UChaosClothAsset::CalculateBounds()
 }
 #endif
 
-void UChaosClothAsset::SetClothCollections(TArray<TSharedRef<const FManagedArrayCollection>>&& InClothCollections)
+void UChaosClothAsset::SetClothCollections(TArray<TSharedRef<const FManagedArrayCollection>>&& InClothCollections, bool bReregisterComponents)
 {
 	WaitUntilAsyncPropertyReleased(EClothAssetAsyncProperties::ClothCollection, ESkinnedAssetAsyncPropertyLockType::WriteOnly);
 
@@ -515,7 +515,7 @@ PRAGMA_DISABLE_DEPRECATION_WARNINGS
 	ClothCollections = MoveTemp(InClothCollections);
 PRAGMA_ENABLE_DEPRECATION_WARNINGS
 
-	OnPropertyChanged();
+	OnPropertyChanged(bReregisterComponents);
 }
 
 void UChaosClothAsset::Build(
diff --git a/Engine/Plugins/ChaosClothAsset/Source/ChaosClothAssetEngine/Private/ChaosClothAsset/ClothAssetBase.cpp b/Engine/Plugins/ChaosClothAsset/Source/ChaosClothAssetEngine/Private/ChaosClothAsset/ClothAssetBase.cpp
--- a/Engine/Plugins/ChaosClothAsset/Source/ChaosClothAssetEngine/Private/ChaosClothAsset/ClothAssetBase.cpp
+++ b/Engine/Plugins/ChaosClothAsset/Source/ChaosClothAssetEngine/Private/ChaosClothAsset/ClothAssetBase.cpp
@@ -729,8 +729,11 @@ bool UChaosClothAssetBase::ExportToSkeletalMesh(USkeletalMesh& SkeletalMesh) con
 }
 #endif  // #if WITH_EDITOR
 
-void UChaosClothAssetBase::OnPropertyChanged() const
+void UChaosClothAssetBase::OnPropertyChanged(const bool bReregisterComponents) const
 {
+	// Context will go out of scope, causing the components to be re-registered, but only when bReregisterComponents is true
+	const FMultiComponentReregisterContext MultiComponentReregisterContext(bReregisterComponents ? GetDependentComponents() : TArray<UActorComponent*>());
+
 	// Update cloth components properties
 	for (TObjectIterator<UChaosClothComponent> ObjectIterator; ObjectIterator; ++ObjectIterator)
 	{
diff --git a/Engine/Plugins/ChaosClothAsset/Source/ChaosClothAssetEngine/Public/ChaosClothAsset/ClothAsset.h b/Engine/Plugins/ChaosClothAsset/Source/ChaosClothAssetEngine/Public/ChaosClothAsset/ClothAsset.h
--- a/Engine/Plugins/ChaosClothAsset/Source/ChaosClothAssetEngine/Public/ChaosClothAsset/ClothAsset.h
+++ b/Engine/Plugins/ChaosClothAsset/Source/ChaosClothAssetEngine/Public/ChaosClothAsset/ClothAsset.h
@@ -92,12 +92,16 @@ public:
 		PRAGMA_ENABLE_DEPRECATION_WARNINGS
 	}
 
-	/** Set new Cloth Collections for this Cloth Asset. */
-	UE_API void SetClothCollections(TArray<TSharedRef<const FManagedArrayCollection>>&& InClothCollections);
+	/**
+	 * Set new Cloth Collections for this Cloth Asset.
+	 * @param InClothCollections The source managed array collections to set to the asset.
+	 * @param bReregisterComponents Whether to reregister all components using this asset. This is the safe way to hard reset all simulations.
+	 */
+	UE_API void SetClothCollections(TArray<TSharedRef<const FManagedArrayCollection>>&& InClothCollections, bool bReregisterComponents = false);
 
 	/**
 	 * Rebuild this asset from the specified managed array collection data.
-	 * @param InClothCollections The source managed array collection to build/rebuild the asset with.
+	 * @param InClothCollections The source managed array collections to build/rebuild the asset with.
 	 * @param InOutTransitionCache The LOD transition cache to speed up the build process, if available.
 	 * @param ErrorText The headline error text if any error occurs, an empty text otherwise.
 	 * @param VerboseText The verbose description of the error if any.
diff --git a/Engine/Plugins/ChaosClothAsset/Source/ChaosClothAssetEngine/Public/ChaosClothAsset/ClothAssetBase.h b/Engine/Plugins/ChaosClothAsset/Source/ChaosClothAssetEngine/Public/ChaosClothAsset/ClothAssetBase.h
--- a/Engine/Plugins/ChaosClothAsset/Source/ChaosClothAssetEngine/Public/ChaosClothAsset/ClothAssetBase.h
+++ b/Engine/Plugins/ChaosClothAsset/Source/ChaosClothAssetEngine/Public/ChaosClothAsset/ClothAssetBase.h
@@ -279,7 +279,7 @@ public:
 
 protected:
 	/** Update all properties for components currently using this asset. */
-	UE_API void OnPropertyChanged() const;
+	UE_API void OnPropertyChanged(const bool bReregisterComponents = false) const;
 
 	/** Update all components currently using this asset. */
 	UE_API void OnAssetChanged(const bool bReregisterComponents = true) const;
diff --git a/Engine/Plugins/ChaosClothAssetDataflowNodes/Source/ChaosClothAssetDataflowNodes/Private/ChaosClothAsset/TerminalNode.cpp b/Engine/Plugins/ChaosClothAssetDataflowNodes/Source/ChaosClothAssetDataflowNodes/Private/ChaosClothAsset/TerminalNode.cpp
--- a/Engine/Plugins/ChaosClothAssetDataflowNodes/Source/ChaosClothAssetDataflowNodes/Private/ChaosClothAsset/TerminalNode.cpp
+++ b/Engine/Plugins/ChaosClothAssetDataflowNodes/Source/ChaosClothAssetDataflowNodes/Private/ChaosClothAsset/TerminalNode.cpp
@@ -67,16 +67,18 @@ namespace UE::Chaos::ClothAsset::Private
 		return Checksum;
 	}
 
-	bool PropertyKeysAndSolverTypesMatch(const TArray<TSharedRef<const FManagedArrayCollection>>& Collections0, const TArray<TSharedRef<const FManagedArrayCollection>>& Collections1)
+	bool ArePropertiesPreserved(const TArray<TSharedRef<const FManagedArrayCollection>>& Collections0, const TArray<TSharedRef<const FManagedArrayCollection>>& Collections1, bool& bNeedsReset)
 	{
+		using namespace ::Chaos::Softs;
+		bNeedsReset = false;
 		if (Collections0.Num() != Collections1.Num())
 		{
 			return false;
 		}
 		for (int32 LODIndex = 0; LODIndex < Collections0.Num(); ++LODIndex)
 		{
-			::Chaos::Softs::FCollectionPropertyConstFacade Property0(Collections0[LODIndex].ToSharedPtr());
-			::Chaos::Softs::FCollectionPropertyConstFacade Property1(Collections1[LODIndex].ToSharedPtr());
+			FCollectionPropertyConstFacade Property0(Collections0[LODIndex].ToSharedPtr());
+			FCollectionPropertyConstFacade Property1(Collections1[LODIndex].ToSharedPtr());
 			if (Property0.Num() != Property1.Num())
 			{
 				return false;
@@ -87,6 +89,20 @@ namespace UE::Chaos::ClothAsset::Private
 				{
 					return false;
 				}
+				if (Property0.IsIntrinsic(PropertyIndex) || Property1.IsIntrinsic(PropertyIndex))
+				{
+					if (!FCollectionPropertyConstFacade::AreValuesEqual(Property0, Property1, PropertyIndex))
+					{
+						return false;  // Intrinsic property change needs full rebuild
+					}
+				}
+				if (!Property0.IsAnimatable(PropertyIndex) || !Property1.IsAnimatable(PropertyIndex))
+				{
+					if (!FCollectionPropertyConstFacade::AreValuesEqual(Property0, Property1, PropertyIndex))
+					{
+						bNeedsReset = true;  // Non animatable property change needs hard reset
+					}
+				}
 			}
 		}
 		return true;
@@ -125,7 +141,9 @@ void FChaosClothAssetTerminalNode_v2::SetAssetValue(TObjectPtr<UObject> Asset, U
 		ClothColllectionChecksum = Private::CalculateClothChecksum(InClothCollections);
 		bClothCollectionChecksumValid = InClothCollections.Num() > 0;
 
-		if (bPreviousChecksumsValid && PreviousChecksum == ClothColllectionChecksum && Private::PropertyKeysAndSolverTypesMatch(InClothCollections, ClothCollections))
+		bool bNeedsReset;
+		if (bPreviousChecksumsValid && PreviousChecksum == ClothColllectionChecksum &&
+			Private::ArePropertiesPreserved(InClothCollections, ClothCollections, bNeedsReset))
 		{
 			// Cloth and property keys match. Just update property values.
 			check(InClothCollections.Num() == ClothCollections.Num());
@@ -138,7 +156,7 @@ void FChaosClothAssetTerminalNode_v2::SetAssetValue(TObjectPtr<UObject> Asset, U
 				ClothCollections[LODIndex] = MoveTemp(ClothCollection);
 			}
 
-			ClothAsset->SetClothCollections(MoveTemp(ClothCollections));
+			ClothAsset->SetClothCollections(MoveTemp(ClothCollections), bNeedsReset);
 
 			// Asset must be resaved
 			ClothAsset->MarkPackageDirty();
@@ -323,7 +341,9 @@ PRAGMA_ENABLE_DEPRECATION_WARNINGS
 		ClothColllectionChecksum = Private::CalculateClothChecksum(InClothCollections);
 		bClothCollectionChecksumValid = InClothCollections.Num() > 0;
 
-		if (bPreviousChecksumsValid && PreviousChecksum == ClothColllectionChecksum && Private::PropertyKeysAndSolverTypesMatch(InClothCollections, RefClothCollections))
+		bool bNeedsReset;
+		if (bPreviousChecksumsValid && PreviousChecksum == ClothColllectionChecksum &&
+			Private::ArePropertiesPreserved(InClothCollections, RefClothCollections, bNeedsReset))
 		{
 			// Cloth and property keys match. Just update property values.
 			TArray<TSharedRef<const FManagedArrayCollection>> ClothCollections = RefClothCollections;  // Copy on write
@@ -337,7 +357,7 @@ PRAGMA_ENABLE_DEPRECATION_WARNINGS
 				ClothCollections[LODIndex] = MoveTemp(ClothCollection);
 			}
 
-			ClothAsset->SetClothCollections(MoveTemp(ClothCollections));
+			ClothAsset->SetClothCollections(MoveTemp(ClothCollections), bNeedsReset);
 
 			// Asset must be resaved
 			ClothAsset->MarkPackageDirty();
diff --git a/Engine/Source/Runtime/Experimental/Chaos/Public/Chaos/CollectionPropertyFacade.h b/Engine/Source/Runtime/Experimental/Chaos/Public/Chaos/CollectionPropertyFacade.h
--- a/Engine/Source/Runtime/Experimental/Chaos/Public/Chaos/CollectionPropertyFacade.h
+++ b/Engine/Source/Runtime/Experimental/Chaos/Public/Chaos/CollectionPropertyFacade.h
@@ -301,6 +301,21 @@ namespace Chaos::Softs
 			return Hash;
 		}
 
+		/** Return whether two properties from different facades have the same values regardless of their names or types. */
+		static bool AreValuesEqual(const FCollectionPropertyConstFacade& Facade, int32 KeyIndex, const FCollectionPropertyConstFacade& Other, int32 OtherKeyIndex)
+		{
+			return
+				Facade.GetLowValue<FVector3f>(KeyIndex) == Other.GetLowValue<FVector3f>(OtherKeyIndex) &&
+				Facade.GetHighValue<FVector3f>(KeyIndex) == Other.GetHighValue<FVector3f>(OtherKeyIndex) &&
+				Facade.GetStringValue(KeyIndex) == Other.GetStringValue(OtherKeyIndex) &&
+				EnumOnlyContainsFlags(Facade.GetFlags(KeyIndex) ^ Other.GetFlags(OtherKeyIndex), ECollectionPropertyFlags::Dirty | ECollectionPropertyFlags::StringDirty);
+		}
+
+		/** Return whether two properties from different facades and sharing the same key index have the same values regardless of their names or types. */
+		static bool AreValuesEqual(const FCollectionPropertyConstFacade& Facade, const FCollectionPropertyConstFacade& Other, int32 KeyIndex)
+		{
+			return AreValuesEqual(Facade, KeyIndex, Other, KeyIndex);
+		}
 
 	protected:
 		// No init constructor for FCollectionPropertyFacade
