Chaos Cloth: Reproject accessory mesh data after RemeshSimMesh

After remeshing the sim mesh, reprojection of accessory mesh vertex data (Position3D, Normal, BoneIndices, BoneWeights) using the per-vertex barycentric coordinates computed during 3D position reprojection.

#Synced-CL 51472478

diff --git a/Engine/Plugins/ChaosClothAssetDataflowNodes/Source/ChaosClothAssetDataflowNodes/ChaosClothAssetDataflowNodes.Build.cs b/Engine/Plugins/ChaosClothAssetDataflowNodes/Source/ChaosClothAssetDataflowNodes/ChaosClothAssetDataflowNodes.Build.cs
--- a/Engine/Plugins/ChaosClothAssetDataflowNodes/Source/ChaosClothAssetDataflowNodes/ChaosClothAssetDataflowNodes.Build.cs
+++ b/Engine/Plugins/ChaosClothAssetDataflowNodes/Source/ChaosClothAssetDataflowNodes/ChaosClothAssetDataflowNodes.Build.cs
@@ -19,6 +19,7 @@ public class ChaosClothAssetDataflowNodes : ModuleRules
 		PrivateDependencyModuleNames.AddRange(
 			new string[]
 			{
+				"AnimationCore",
 				"AssetTools",
 				"ChaosCaching",
 				"ChaosCloth",
diff --git a/Engine/Plugins/ChaosClothAssetDataflowNodes/Source/ChaosClothAssetDataflowNodes/Private/ChaosClothAsset/RemeshNode.cpp b/Engine/Plugins/ChaosClothAssetDataflowNodes/Source/ChaosClothAssetDataflowNodes/Private/ChaosClothAsset/RemeshNode.cpp
--- a/Engine/Plugins/ChaosClothAssetDataflowNodes/Source/ChaosClothAssetDataflowNodes/Private/ChaosClothAsset/RemeshNode.cpp
+++ b/Engine/Plugins/ChaosClothAssetDataflowNodes/Source/ChaosClothAssetDataflowNodes/Private/ChaosClothAsset/RemeshNode.cpp
@@ -3,6 +3,7 @@
 #include "ChaosClothAsset/RemeshNode.h"
 
 #include "Algo/Find.h"
+#include "BoneWeights.h"
 #include "Algo/RemoveIf.h"
 #include "Chaos/CollectionEmbeddedSpringConstraintFacade.h"
 #include "Chaos/CollectionPropertyFacade.h"
@@ -12,6 +13,7 @@
 #include "ChaosClothAsset/ClothEngineTools.h"
 #include "ChaosClothAsset/ClothGeometryTools.h"
 #include "ChaosClothAsset/ClothPatternToDynamicMesh.h"
+#include "ChaosClothAsset/ClothPatternToDynamicMeshMappingSupport.h"
 #include "ChaosClothAsset/CollectionClothFacade.h"
 #include "ChaosClothAsset/CollectionClothSelectionFacade.h"
 #include "ChaosClothAsset/SimulationSelfCollisionSpheresConfigNode.h"
@@ -1363,6 +1365,131 @@ namespace UE::Chaos::ClothAsset::Private
 		FVector2f DensityLowHigh = { 0,1 };
 	};
 
+	struct FVertexProjectionData
+	{
+		UE::Geometry::FIndex3i GlobalTriangle;  // Global sim vertex 3D indices from ProjectionTarget3D
+		FVector3d BaryCoords;
+	};
+
+	static void ReprojectAccessoryMeshes(
+		const TSharedRef<const FManagedArrayCollection>& InClothCollection,
+		const TSharedRef<FManagedArrayCollection>& OutClothCollection,
+		const TMap<int32, FVertexProjectionData>& VertexProjectionDataMap,
+		const TMap<int, int32>& DynamicMeshToClothVertexMap)
+	{
+		using namespace UE::Chaos::ClothAsset;
+
+		FCollectionClothConstFacade InClothFacade(InClothCollection);
+		FCollectionClothFacade OutClothFacade(OutClothCollection);
+
+		const int32 NumAccessoryMeshes = InClothFacade.GetNumSimAccessoryMeshes();
+		if (NumAccessoryMeshes == 0)
+		{
+			return;
+		}
+
+		// Clear stale data that was CopyTo'd with the old vertex count
+		OutClothFacade.SetNumSimAccessoryMeshes(0);
+
+		const int32 NumOutputVertices = OutClothFacade.GetNumSimVertices3D();
+
+		for (int32 MeshIdx = 0; MeshIdx < NumAccessoryMeshes; ++MeshIdx)
+		{
+			FCollectionClothSimAccessoryMeshConstFacade InAccessory = InClothFacade.GetSimAccessoryMesh(MeshIdx);
+
+			const TConstArrayView<FVector3f> InPositions = InAccessory.GetSimAccessoryMeshPosition3D();
+			const TConstArrayView<FVector3f> InNormals = InAccessory.GetSimAccessoryMeshNormal();
+			const TConstArrayView<TArray<int32>> InBoneIndices = InAccessory.GetSimAccessoryMeshBoneIndices();
+			const TConstArrayView<TArray<float>> InBoneWeights = InAccessory.GetSimAccessoryMeshBoneWeights();
+
+			const bool bHasSkinning = InBoneIndices.Num() > 0 && InBoneWeights.Num() > 0 &&
+				InBoneIndices.Num() == InPositions.Num() && InBoneWeights.Num() == InPositions.Num();
+
+			TArray<FVector3f> OutPositions;
+			TArray<FVector3f> OutNormals;
+			TArray<TArray<int32>> OutBoneIndices;
+			TArray<TArray<float>> OutBoneWeights;
+
+			OutPositions.SetNumZeroed(NumOutputVertices);
+			OutNormals.SetNumZeroed(NumOutputVertices);
+			OutBoneIndices.SetNum(NumOutputVertices);
+			OutBoneWeights.SetNum(NumOutputVertices);
+
+			for (const TPair<int, int32>& Pair : DynamicMeshToClothVertexMap)
+			{
+				const int32 VID = Pair.Key;
+				const int32 OutVertIdx = Pair.Value;
+
+				const FVertexProjectionData* ProjData = VertexProjectionDataMap.Find(VID);
+				if (!ProjData)
+				{
+					continue;
+				}
+
+				const int32 G0 = ProjData->GlobalTriangle[0];
+				const int32 G1 = ProjData->GlobalTriangle[1];
+				const int32 G2 = ProjData->GlobalTriangle[2];
+
+				if (!InPositions.IsValidIndex(G0) || !InPositions.IsValidIndex(G1) || !InPositions.IsValidIndex(G2) || !OutPositions.IsValidIndex(OutVertIdx))
+				{
+					continue;
+				}
+
+				const FVector3d& B = ProjData->BaryCoords;
+
+				// Interpolate Position3D
+				OutPositions[OutVertIdx] = FVector3f(
+					(float)B[0] * InPositions[G0] +
+					(float)B[1] * InPositions[G1] +
+					(float)B[2] * InPositions[G2]);
+
+				// Interpolate Normal
+				FVector3f InterpNormal =
+					(float)B[0] * InNormals[G0] +
+					(float)B[1] * InNormals[G1] +
+					(float)B[2] * InNormals[G2];
+				InterpNormal.Normalize();
+				OutNormals[OutVertIdx] = InterpNormal;
+
+				// Interpolate skinning data
+				if (bHasSkinning)
+				{
+					auto MakeBoneWeights = [](const TArray<int32>& Indices, const TArray<float>& Weights)
+					{
+						check(Indices.Num() == Weights.Num());
+
+						TArray<FBoneIndexType> BoneIndexTypes;
+						BoneIndexTypes.Reserve(Indices.Num());
+						for (const int32 Idx : Indices)
+						{
+							BoneIndexTypes.Add(static_cast<FBoneIndexType>(Idx));
+						}
+						return UE::AnimationCore::FBoneWeights::Create(BoneIndexTypes.GetData(), Weights.GetData(), Indices.Num());
+					};
+
+					const UE::AnimationCore::FBoneWeights BlendedWeights = UE::AnimationCore::FBoneWeights::Blend(
+						MakeBoneWeights(InBoneIndices[G0], InBoneWeights[G0]),
+						MakeBoneWeights(InBoneIndices[G1], InBoneWeights[G1]),
+						MakeBoneWeights(InBoneIndices[G2], InBoneWeights[G2]),
+						(float)B[0], (float)B[1], (float)B[2]);
+
+					TArray<int32>& OutIndices = OutBoneIndices[OutVertIdx];
+					TArray<float>& OutWeights = OutBoneWeights[OutVertIdx];
+					OutIndices.Reserve(BlendedWeights.Num());
+					OutWeights.Reserve(BlendedWeights.Num());
+					for (int32 i = 0; i < BlendedWeights.Num(); ++i)
+					{
+						OutIndices.Add(static_cast<int32>(BlendedWeights[i].GetBoneIndex()));
+						OutWeights.Add(BlendedWeights[i].GetWeight());
+					}
+				}
+			}
+
+			FCollectionClothSimAccessoryMeshFacade OutAccessory = OutClothFacade.AddGetSimAccessoryMesh();
+			OutAccessory.Initialize(InAccessory.GetSimAccessoryMeshName(), OutPositions, OutNormals, OutBoneIndices, OutBoneWeights);
+		}
+	}
+
 	static void RemeshSimMesh(
 		const TSharedRef<const FManagedArrayCollection>& ClothCollection,
 		const FRemeshSimMeshParams& Params,
@@ -1480,6 +1607,8 @@ namespace UE::Chaos::ClothAsset::Private
 		FDynamicMesh3 Mesh3D;
 		Mesh3D.Copy(Mesh2D);
 
+		TMap<int32, FVertexProjectionData> VertexProjectionDataMap;
+
 		TMap<int32, TSet<int32>> PatternVertexIDs;
 		for (int32 TID : Mesh2D.TriangleIndicesItr())
 		{
@@ -1506,6 +1635,7 @@ namespace UE::Chaos::ClothAsset::Private
 			Converter.Convert(ClothCollection, PatternID, EClothPatternVertexType::Sim2D, ProjectionTarget2D);
 			FDynamicMesh3 ProjectionTarget3D;
 			Converter.Convert(ClothCollection, PatternID, EClothPatternVertexType::Sim3D, ProjectionTarget3D);
+			const FClothPatternToDynamicMeshMappingSupport MappingSupport3D(ProjectionTarget3D);
 
 			if (ProjectionTarget2D.TriangleCount() < ProjectionTarget3D.TriangleCount())
 			{
@@ -1581,6 +1711,12 @@ namespace UE::Chaos::ClothAsset::Private
 				const FVector3d InterpolatedPoint = ProjectionTarget3D.GetTriBaryPoint(NearestTriangle, Bary[0], Bary[1], Bary[2]);
 
 				Mesh3D.SetVertex(VID, InterpolatedPoint);
+
+				const FIndex3i GlobalTriangle(
+					MappingSupport3D.GetOriginalVertexID(Triangle3D[0]),
+					MappingSupport3D.GetOriginalVertexID(Triangle3D[1]),
+					MappingSupport3D.GetOriginalVertexID(Triangle3D[2]));
+				VertexProjectionDataMap.Add(VID, { GlobalTriangle, Bary });
 			}
 		}
 
@@ -1592,6 +1728,9 @@ namespace UE::Chaos::ClothAsset::Private
 		TMap<int, int32> DynamicMeshToClothVertexMap;
 		FClothGeometryTools::BuildSimMeshFromDynamicMeshes(OutClothCollection, Mesh2D, Mesh3D, PatternIndexLayerID, bTransferWeightMaps, bTransferSimSkinningData, bAppendToExistingMesh, DynamicMeshToClothVertexMap);
 
+		// Reproject accessory mesh data before the seam rebuild, while DynamicMeshToClothVertexMap is still valid.
+		// Seam initialization modifies SimVertices3D (splits boundary vertices), which invalidates the map.
+		Private::ReprojectAccessoryMeshes(ClothCollection, OutClothCollection, VertexProjectionDataMap, DynamicMeshToClothVertexMap);
 
 		// Re-apply the seam info from the input sim mesh. This will create a new Seam for each set of connected stitches
 
diff --git a/Engine/Plugins/ChaosClothAssetDataflowNodes/Source/ChaosClothAssetDataflowNodes/Public/ChaosClothAsset/RemeshNode.h b/Engine/Plugins/ChaosClothAssetDataflowNodes/Source/ChaosClothAssetDataflowNodes/Public/ChaosClothAsset/RemeshNode.h
--- a/Engine/Plugins/ChaosClothAssetDataflowNodes/Source/ChaosClothAssetDataflowNodes/Public/ChaosClothAsset/RemeshNode.h
+++ b/Engine/Plugins/ChaosClothAssetDataflowNodes/Source/ChaosClothAssetDataflowNodes/Public/ChaosClothAsset/RemeshNode.h
@@ -25,7 +25,7 @@ enum class EChaosClothAssetRemeshSelectionSource : uint8
 };
 
 /** Remesh the cloth surface(s) to get the specified mesh resolution(s).
- *  NOTE: Weight Maps, Skinning Data, Self Collision Spheres, and Long Range Attachment Constraints will be reconstructed on the output mesh, however all other Selections will be removed
+ *  NOTE: Accessory Meshes, Weight Maps, Skinning Data, Self Collision Spheres, and Long Range Attachment Constraints will be reconstructed on the output mesh, however all other Selections will be removed
  */
 USTRUCT(Meta = (DataflowCloth))
 struct FChaosClothAssetRemeshNode_v2 : public FDataflowNode
