Dataflow : cloth : add support for display vertex attribute when the node request it

#Synced-CL 51967769

diff --git a/Engine/Plugins/ChaosClothAssetDataflowNodes/Source/ChaosClothAssetDataflowNodes/Private/ChaosClothAsset/ClothRenderableTypes.cpp b/Engine/Plugins/ChaosClothAssetDataflowNodes/Source/ChaosClothAssetDataflowNodes/Private/ChaosClothAsset/ClothRenderableTypes.cpp
--- a/Engine/Plugins/ChaosClothAssetDataflowNodes/Source/ChaosClothAssetDataflowNodes/Private/ChaosClothAsset/ClothRenderableTypes.cpp
+++ b/Engine/Plugins/ChaosClothAssetDataflowNodes/Source/ChaosClothAssetDataflowNodes/Private/ChaosClothAsset/ClothRenderableTypes.cpp
@@ -90,6 +90,54 @@ namespace UE::Chaos::ClothAsset
 			}
 		}
 
+		static void AddAttributeAsColorToMesh(const FCollectionClothConstFacade& ClothFacade, FDynamicMesh3& InOutMesh, TConstArrayView<float> AttributeValue)
+		{
+			using namespace UE::Geometry;
+
+			checkf(InOutMesh.VertexCount() == AttributeValue.Num(), TEXT("Expected to have the same number of vertices in LodMesh and ClothCollection mesh"));
+
+			if (!InOutMesh.Attributes()->HasPrimaryColors())
+			{
+				InOutMesh.Attributes()->EnablePrimaryColors();
+				InOutMesh.Attributes()->PrimaryColors()->CreateFromPredicate([](int ParentVID, int TriIDA, int TriIDB)
+					{
+						return false; // make sure we have a value per corner
+					}
+				, 0.0f);
+			}
+
+
+			if (FDynamicMeshColorOverlay* const ColorAttributeLayer = InOutMesh.Attributes()->PrimaryColors())
+			{
+				TConstArrayView<FIntVector> Triangles = ClothFacade.GetSimIndices3D();
+
+				auto GetVertexColor = [&AttributeValue](int32 VtxIdx)
+					{
+						if (AttributeValue.IsValidIndex(VtxIdx))
+						{
+							const float Value = AttributeValue[VtxIdx];
+							return FVector4f{ Value , Value , Value, 1.0f };
+						}
+						return FVector4f{ 1, 0, 0, 1.0f };
+					};
+
+				for (int32 TriID = 0; TriID < Triangles.Num(); ++TriID)
+				{
+					FIndex3i AttrTri;
+					if (ColorAttributeLayer->GetTriangleIfValid(TriID, AttrTri))
+					{
+						const int32 VtxIdx0 = ColorAttributeLayer->GetParentVertex(AttrTri[0]);
+						const int32 VtxIdx1 = ColorAttributeLayer->GetParentVertex(AttrTri[1]);
+						const int32 VtxIdx2 = ColorAttributeLayer->GetParentVertex(AttrTri[2]);
+
+						ColorAttributeLayer->SetElement(AttrTri[0], GetVertexColor(VtxIdx0));
+						ColorAttributeLayer->SetElement(AttrTri[1], GetVertexColor(VtxIdx1));
+						ColorAttributeLayer->SetElement(AttrTri[2], GetVertexColor(VtxIdx2));
+					}
+				}
+			}
+		}
+
 		static bool AddFaceSelectionToMesh(const TSharedRef<const FManagedArrayCollection>& Collection, FName SelectionName, FDynamicMesh3& InOutMesh)
 		{
 			using namespace UE::Geometry;
@@ -120,9 +168,9 @@ namespace UE::Chaos::ClothAsset
 			if (!InOutMesh.Attributes()->HasPrimaryColors())
 			{
 				InOutMesh.Attributes()->EnablePrimaryColors();
-				InOutMesh.Attributes()->PrimaryColors()->CreateFromPredicate([](int ParentVID, int TriIDA, int TriIDB)
+				InOutMesh.Attributes()->PrimaryColors()-> CreateFromPredicate([](int ParentVID, int TriIDA, int TriIDB)
 					{
-						return false; // make sure we have a value per corner
+						return true; // make sure we have a value per corner
 					}
 				, 0.0f);
 			}
@@ -444,10 +492,25 @@ namespace UE::Chaos::ClothAsset
 				FClothPatternToDynamicMesh Converter;
 				Converter.Convert(ClothCollection, INDEX_NONE, EClothPatternVertexType::Sim3D, DynamicMesh);
 
+				bool bNeedVertexColorMaterial = false;
 				const bool bShowPatternColors = (Settings && Settings->bShowPatternColors);
 				if (bShowPatternColors)
 				{
 					Private::AddPatternsColorsToMesh(ClothFacade, DynamicMesh);
+					bNeedVertexColorMaterial = true;
+				}
+				else 
+				{
+					const FDataflowNode::FAttributeKey AttributeKey = Instance.GetVertexAttributeToVisualize();
+					if (AttributeKey.GroupName == ClothCollectionGroup::SimVertices3D)
+					{
+						TConstArrayView<float> MapValues = ClothFacade.GetWeightMap(AttributeKey.AttributeName);
+						if (MapValues.Num() > 0)
+						{
+							Private::AddAttributeAsColorToMesh(ClothFacade, DynamicMesh, MapValues);
+							bNeedVertexColorMaterial = true;
+						}
+					}
 				}
 				const FString SelectionName = Instance.GetOutputValueByType<FString>({}, TEXT("DataflowSelection"));
 				const bool bHasFaceSelection = Private::AddFaceSelectionToMesh(ClothCollection, FName(SelectionName), DynamicMesh);
@@ -464,7 +527,7 @@ namespace UE::Chaos::ClothAsset
 						Component->SetCastShadow(false);
 						Component->SetMesh(MoveTemp(DynamicMesh));
 
-						UMaterialInterface* TwoSidedMaterial = bShowPatternColors || bHasFaceSelection
+						UMaterialInterface* TwoSidedMaterial = bNeedVertexColorMaterial || bHasFaceSelection
 							? Private::GetTwoSidedMaterialForColoredPatterns()
 							: Private::GetTwoSidedMaterialForRenderingTypes();
 						Component->SetOverrideRenderMaterial(TwoSidedMaterial);
diff --git a/Engine/Plugins/Experimental/Dataflow/Source/DataflowEditor/Public/DataflowRendering/DataflowRenderableTypeInstance.h b/Engine/Plugins/Experimental/Dataflow/Source/DataflowEditor/Public/DataflowRendering/DataflowRenderableTypeInstance.h
--- a/Engine/Plugins/Experimental/Dataflow/Source/DataflowEditor/Public/DataflowRendering/DataflowRenderableTypeInstance.h
+++ b/Engine/Plugins/Experimental/Dataflow/Source/DataflowEditor/Public/DataflowRendering/DataflowRenderableTypeInstance.h
@@ -84,7 +84,7 @@ namespace UE::Dataflow
 			}
 		}
 
-		const FDataflowNode::FAttributeKey GetVertexAttributeToVisualize() const;
+		UE_API const FDataflowNode::FAttributeKey GetVertexAttributeToVisualize() const;
 
 		const void GetSelectionToVisualize(FDataflowSelection& OutSelection) const;
 
