Chaos Cloth Asset - Fixed crash when changing Mesh Targets while painting.

#rnx

#Synced-CL 51673482

diff --git a/Engine/Plugins/ChaosClothAssetDataflowNodes/Source/ChaosClothAssetDataflowNodes/Private/ChaosClothAsset/WeightMapNode.cpp b/Engine/Plugins/ChaosClothAssetDataflowNodes/Source/ChaosClothAssetDataflowNodes/Private/ChaosClothAsset/WeightMapNode.cpp
--- a/Engine/Plugins/ChaosClothAssetDataflowNodes/Source/ChaosClothAssetDataflowNodes/Private/ChaosClothAsset/WeightMapNode.cpp
+++ b/Engine/Plugins/ChaosClothAssetDataflowNodes/Source/ChaosClothAssetDataflowNodes/Private/ChaosClothAsset/WeightMapNode.cpp
@@ -338,11 +338,22 @@ void FChaosClothAssetWeightMapNode::SetVertexAttributeValues(UE::Dataflow::FCont
 		int32 NumFinalWeights;
 		if (MeshTarget == EChaosClothAssetWeightMapMeshTarget::Simulation)
 		{
+			if ((InValues.Num() != ClothFacade.GetNumSimVertices3D()) &&
+				(InValues.Num() != ClothFacade.GetNumSimVertices2D() || InWeightIndices.Num() != ClothFacade.GetNumSimVertices2D()))
+			{
+				Context.Error(LOCTEXT("MeshTargetChangedToRender", "Can't match the painted vertex back to the Cloth Collection, most likely because the Mesh Target has been changed to Render while painting a Simulation weight map."));
+				return;
+			}
 			InputAttributeValues = ClothFacade.GetWeightMap(InInputName);
 			NumFinalWeights = ClothFacade.GetNumSimVertices3D();
 		}
 		else
 		{
+			if (InValues.Num() != ClothFacade.GetNumRenderVertices())
+			{
+				Context.Error(LOCTEXT("MeshTargetChangedToSim", "Can't match the painted vertex back to the Cloth Collection, most likely because the Mesh Target has been changed to Simulation while painting a Render weight map."));
+				return;
+			}
 			InputAttributeValues = ClothFacade.GetUserDefinedAttribute<float>(InInputName, ClothCollectionGroup::RenderVertices);
 			NumFinalWeights = ClothFacade.GetNumRenderVertices();
 		}
