Dataflow : add collection facade to access basic data from a vertex group ( see vertex group property )
- it is used by the tools to properly store and retrieve snapshot with vertex and triangle information
- also define the target info for cloth and geometry collection  target groups

#Synced-CL 52197375

diff --git a/Engine/Plugins/ChaosClothAssetDataflowNodes/Source/ChaosClothAssetDataflowNodes/Private/ChaosClothAsset/ChaosClothAssetDataflowNodesModule.cpp b/Engine/Plugins/ChaosClothAssetDataflowNodes/Source/ChaosClothAssetDataflowNodes/Private/ChaosClothAsset/ChaosClothAssetDataflowNodesModule.cpp
--- a/Engine/Plugins/ChaosClothAssetDataflowNodes/Source/ChaosClothAssetDataflowNodes/Private/ChaosClothAsset/ChaosClothAssetDataflowNodesModule.cpp
+++ b/Engine/Plugins/ChaosClothAssetDataflowNodes/Source/ChaosClothAssetDataflowNodes/Private/ChaosClothAsset/ChaosClothAssetDataflowNodesModule.cpp
@@ -1,5 +1,6 @@
 // Copyright Epic Games, Inc. All Rights Reserved.
 
+#include "Chaos/Range.h"
 #include "ChaosClothAsset/AddStitchNode.h"
 #include "ChaosClothAsset/AddWeightMapNode.h"
 #include "ChaosClothAsset/ApplyProxyDeformerNode.h"
@@ -8,6 +9,7 @@
 #include "ChaosClothAsset/BindToRootBoneNode.h"
 #include "ChaosClothAsset/BlendVerticesNode.h"
 #include "ChaosClothAsset/ClothAsset.h"
+#include "ChaosClothAsset/ClothCollectionAttribute.h"
 #include "ChaosClothAsset/ClothCollectionGroup.h"
 #include "ChaosClothAsset/ClothCollectionToDynamicMeshNode.h"
 #include "ChaosClothAsset/ClothCollectionQueryNode.h"
@@ -424,7 +426,7 @@ namespace UE::Chaos::ClothAsset
 			UE::Dataflow::FRenderingViewModeFactory::GetInstance().DeregisterViewMode(FClothRenderViewMode::Name);
 		}
 
-
+		////////////////////////////////////////////////////////////////////////////////////////////////////////////////////////////////////
 
 		class FClothCollectionAddScalarVertexPropertyCallbacks : public IDataflowAddScalarVertexPropertyCallbacks
 		{
@@ -444,6 +446,29 @@ namespace UE::Chaos::ClothAsset
 				return { ClothCollectionGroup::SimVertices2D, ClothCollectionGroup::SimVertices3D, ClothCollectionGroup::RenderVertices };
 			}
 
+			virtual void GetTargetGroupInfos(TArray<FTargetGroupInfo>& OutInfos) const override
+			{
+				OutInfos.Add({
+					.TargetGroup = ClothCollectionGroup::SimVertices2D,
+					.PositionAttributeKey = { ClothCollectionAttribute::SimPosition3D, ClothCollectionGroup::SimVertices3D },
+					.IndicesAttributeKey = { ClothCollectionAttribute::SimIndices3D, ClothCollectionGroup::SimFaces },
+					.MappingFrom2DTo3DAttributeKey = { ClothCollectionAttribute::SimVertex3DLookup, ClothCollectionGroup::SimVertices2D },
+					});
+				OutInfos.Add({
+					.TargetGroup = ClothCollectionGroup::SimVertices3D,
+					.PositionAttributeKey = { ClothCollectionAttribute::SimPosition3D, ClothCollectionGroup::SimVertices3D },
+					.IndicesAttributeKey = { ClothCollectionAttribute::SimIndices3D, ClothCollectionGroup::SimFaces },
+					.MappingFrom2DTo3DAttributeKey = {},
+					});
+				OutInfos.Add({
+					.TargetGroup = ClothCollectionGroup::RenderVertices,
+					.PositionAttributeKey = { ClothCollectionAttribute::RenderPosition, ClothCollectionGroup::RenderVertices },
+					.IndicesAttributeKey = { ClothCollectionAttribute::RenderIndices, ClothCollectionGroup::RenderFaces },
+					.MappingFrom2DTo3DAttributeKey = {},
+					});
+
+			}
+
 			virtual TArray<UE::Dataflow::FRenderingParameter> GetRenderingParameters() const override
 			{
 				return { {TEXT("SurfaceRender"), FName("FClothCollection"), {TEXT("Collection")}, FCloth2DSimViewMode::Name},
diff --git a/Engine/Plugins/Experimental/Dataflow/Source/DataflowNodes/Private/Dataflow/DataflowCollectionAttributeUtils.cpp b/Engine/Plugins/Experimental/Dataflow/Source/DataflowNodes/Private/Dataflow/DataflowCollectionAttributeUtils.cpp
new file mode 100644
--- /dev/null
+++ b/Engine/Plugins/Experimental/Dataflow/Source/DataflowNodes/Private/Dataflow/DataflowCollectionAttributeUtils.cpp
@@ -0,0 +1,68 @@
+// Copyright Epic Games, Inc. All Rights Reserved.
+
+#include "Dataflow/DataflowCollectionAttributeUtils.h"
+
+#include "Dataflow/DataflowNode.h"
+#include "GeometryCollection/GeometryCollection.h"
+
+
+#define LOCTEXT_NAMESPACE "DataflowCollectionAttributeUtils"
+
+namespace UE::Dataflow
+{
+	FDataflowCollectionFacade::FDataflowCollectionFacade(const FManagedArrayCollection& InCollection, FName TargetGroup)
+		: Collection(InCollection)
+	{
+		TargetGroupInfo = FDataflowAddScalarVertexPropertyCallbackRegistry::Get().GetTargetGroupInfo(TargetGroup);
+	}
+	
+	TConstArrayView<FVector3f> FDataflowCollectionFacade::GetVertexPositions() const
+	{
+		if (TargetGroupInfo.TargetGroup.IsNone())
+		{
+			return {};
+		}
+
+		const FDataflowNode::FAttributeKey Key = TargetGroupInfo.PositionAttributeKey;
+
+		if (const TManagedArray<FVector3f>* VertexAttribute = Collection.FindAttributeTyped<FVector3f>(Key.AttributeName, Key.GroupName))
+		{
+			return VertexAttribute->GetConstArray();
+		}
+		return {};
+	}
+
+	TConstArrayView<FIntVector> FDataflowCollectionFacade::GetTriangleIndices() const
+	{
+		if (TargetGroupInfo.TargetGroup.IsNone())
+		{
+			return {};
+		}
+
+		const FDataflowNode::FAttributeKey Key = TargetGroupInfo.IndicesAttributeKey;
+
+		if (const TManagedArray<FIntVector>* IndicesAttribute = Collection.FindAttributeTyped<FIntVector>(Key.AttributeName, Key.GroupName))
+		{
+			return IndicesAttribute->GetConstArray();
+		}
+		return {};
+	}
+
+	TConstArrayView<int32> FDataflowCollectionFacade::Get2Dto3DMapping() const
+	{
+		if (TargetGroupInfo.TargetGroup.IsNone())
+		{
+			return {};
+		}
+
+		const FDataflowNode::FAttributeKey Key = TargetGroupInfo.MappingFrom2DTo3DAttributeKey;
+		if (const TManagedArray<int32>* IndicesAttribute = Collection.FindAttributeTyped<int32>(Key.AttributeName, Key.GroupName))
+		{
+			return IndicesAttribute->GetConstArray();
+		}
+		return {};
+	}
+}
+
+#undef LOCTEXT_NAMESPACE 
+
diff --git a/Engine/Plugins/Experimental/Dataflow/Source/DataflowNodes/Private/Dataflow/DataflowCollectionAttributeUtils.h b/Engine/Plugins/Experimental/Dataflow/Source/DataflowNodes/Private/Dataflow/DataflowCollectionAttributeUtils.h
new file mode 100644
--- /dev/null
+++ b/Engine/Plugins/Experimental/Dataflow/Source/DataflowNodes/Private/Dataflow/DataflowCollectionAttributeUtils.h
@@ -0,0 +1,29 @@
+// Copyright Epic Games, Inc. All Rights Reserved.
+
+#pragma once 
+
+#include <GeometryCollection/ManagedArrayCollection.h>
+#include <Math/MathFwd.h>
+#include <UObject/NameTypes.h>
+#include "Dataflow/DataflowCollectionAddScalarVertexPropertyNode.h"
+
+namespace UE::Dataflow
+{
+	// General collection facade to help nodes to deal with both Cloth and geometry collections
+	struct FDataflowCollectionFacade
+	{
+	public:
+		FDataflowCollectionFacade(const FManagedArrayCollection& Collection, FName TargetGroup);
+		
+		TConstArrayView<FVector3f> GetVertexPositions() const;
+		TConstArrayView<FIntVector> GetTriangleIndices() const;
+		TConstArrayView<int32> Get2Dto3DMapping() const;
+
+	private:
+		using FTargetGroupInfo = FDataflowAddScalarVertexPropertyCallbackRegistry::FTargetGroupInfo;
+
+		const FManagedArrayCollection& Collection;
+		FTargetGroupInfo TargetGroupInfo;
+	};
+
+}
diff --git a/Engine/Plugins/Experimental/Dataflow/Source/DataflowNodes/Private/Dataflow/DataflowNodesPlugin.cpp b/Engine/Plugins/Experimental/Dataflow/Source/DataflowNodes/Private/Dataflow/DataflowNodesPlugin.cpp
--- a/Engine/Plugins/Experimental/Dataflow/Source/DataflowNodes/Private/Dataflow/DataflowNodesPlugin.cpp
+++ b/Engine/Plugins/Experimental/Dataflow/Source/DataflowNodes/Private/Dataflow/DataflowNodesPlugin.cpp
@@ -46,6 +46,16 @@ public:
 		return { FGeometryCollection::VerticesGroup };
 	}
 
+	virtual void GetTargetGroupInfos(TArray<FTargetGroupInfo>& OutInfos) const override
+	{
+		OutInfos.Add({
+			.TargetGroup = FGeometryCollection::VerticesGroup,
+			.PositionAttributeKey = { FGeometryCollection::VertexPositionAttribute, FGeometryCollection::VerticesGroup  },
+			.IndicesAttributeKey = { FGeometryCollection::FaceIndicesAttribute, FGeometryCollection::FacesGroup },
+			.MappingFrom2DTo3DAttributeKey = {}
+			});
+	}
+
 	virtual TArray<UE::Dataflow::FRenderingParameter> GetRenderingParameters() const override
 	{
 		return { { TEXT("SurfaceRender"), FGeometryCollection::StaticType(), {TEXT("Collection")} } };
