Dataflow : fix issue with the cloth weightmap node not properly transfering the active snapshot data to the output collection

#Synced-CL 52232971

diff --git a/Engine/Plugins/ChaosClothAssetDataflowNodes/Source/ChaosClothAssetDataflowNodes/Private/ChaosClothAsset/WeightMapNode.cpp b/Engine/Plugins/ChaosClothAssetDataflowNodes/Source/ChaosClothAssetDataflowNodes/Private/ChaosClothAsset/WeightMapNode.cpp
--- a/Engine/Plugins/ChaosClothAssetDataflowNodes/Source/ChaosClothAssetDataflowNodes/Private/ChaosClothAsset/WeightMapNode.cpp
+++ b/Engine/Plugins/ChaosClothAssetDataflowNodes/Source/ChaosClothAssetDataflowNodes/Private/ChaosClothAsset/WeightMapNode.cpp
@@ -209,7 +209,7 @@ void FChaosClothAssetWeightMapNode::Evaluate(UE::Dataflow::FContext& Context, co
 
 					constexpr bool bIsSim = true;
 					CheckSourceVertexWeights(ClothSimWeights, WeightMapValues, bIsSim);
-					CalculateFinalVertexWeightValues(ClothFacade.GetWeightMap(InInputName), ClothSimWeights);
+					CalculateFinalVertexWeightValues(&ClothFacade, ClothFacade.GetWeightMap(InInputName), ClothSimWeights);
 				}
 			}
 			else 
@@ -235,7 +235,7 @@ void FChaosClothAssetWeightMapNode::Evaluate(UE::Dataflow::FContext& Context, co
 
 					constexpr bool bIsSim = false;
 					CheckSourceVertexWeights(ClothRenderWeights, WeightMapValues, bIsSim);
-					CalculateFinalVertexWeightValues(ClothFacade.GetUserDefinedAttribute<float>(InInputName, ClothCollectionGroup::RenderVertices), ClothRenderWeights);
+					CalculateFinalVertexWeightValues(&ClothFacade, ClothFacade.GetUserDefinedAttribute<float>(InInputName, ClothCollectionGroup::RenderVertices), ClothRenderWeights);
 				}
 			}
 		}
@@ -423,10 +423,10 @@ void FChaosClothAssetWeightMapNode::SetVertexWeights(const TConstArrayView<float
 	StoreWeightsInSnapshot(/*ClothFacade*/nullptr, WeightsToStore, AddSnapshot());
 }
 
-void FChaosClothAssetWeightMapNode::CalculateFinalVertexWeightValues(const TConstArrayView<float> InputMap, TArrayView<float> FinalOutputMap) const
+void FChaosClothAssetWeightMapNode::CalculateFinalVertexWeightValues(UE::Chaos::ClothAsset::FCollectionClothConstFacade* ClothFacade, const TConstArrayView<float> InputMap, TArrayView<float> FinalOutputMap) const
 {
 	TArray<float> WeightMapValues;
-	GetWeightsFromActiveSnapshot(/*ClothFacade*/nullptr, WeightMapValues);
+	GetWeightsFromActiveSnapshot(ClothFacade, WeightMapValues);
 
 	UE::Chaos::ClothAsset::Private::CalculateFinalVertexWeightValues(InputMap, FinalOutputMap, MapOverrideType, WeightMapValues);
 }
diff --git a/Engine/Plugins/ChaosClothAssetDataflowNodes/Source/ChaosClothAssetDataflowNodes/Public/ChaosClothAsset/WeightMapNode.h b/Engine/Plugins/ChaosClothAssetDataflowNodes/Source/ChaosClothAssetDataflowNodes/Public/ChaosClothAsset/WeightMapNode.h
--- a/Engine/Plugins/ChaosClothAssetDataflowNodes/Source/ChaosClothAssetDataflowNodes/Public/ChaosClothAsset/WeightMapNode.h
+++ b/Engine/Plugins/ChaosClothAssetDataflowNodes/Source/ChaosClothAssetDataflowNodes/Public/ChaosClothAsset/WeightMapNode.h
@@ -90,7 +90,7 @@ private:
 	void CHAOSCLOTHASSETDATAFLOWNODES_API SetVertexWeights(const TConstArrayView<float> InputMap, const TArray<float>& FinalValues);
 
 	// Input and FinalOutputMap can be the same array, but should not otherwise be interleaved.
-	void CHAOSCLOTHASSETDATAFLOWNODES_API CalculateFinalVertexWeightValues(const TConstArrayView<float> InputMap, TArrayView<float> FinalOutputMap) const;
+	void CHAOSCLOTHASSETDATAFLOWNODES_API CalculateFinalVertexWeightValues(UE::Chaos::ClothAsset::FCollectionClothConstFacade* ClothFacade, const TConstArrayView<float> InputMap, TArrayView<float> FinalOutputMap) const;
 
 	//~ Begin FDataflowNode interface
 	virtual void Evaluate(UE::Dataflow::FContext& Context, const FDataflowOutput* Out) const override;
diff --git a/Engine/Plugins/ChaosClothAssetEditorCore/Source/ChaosClothAssetEditorTools/Private/ChaosClothAsset/ClothWeightMapPaintTool.cpp b/Engine/Plugins/ChaosClothAssetEditorCore/Source/ChaosClothAssetEditorTools/Private/ChaosClothAsset/ClothWeightMapPaintTool.cpp
--- a/Engine/Plugins/ChaosClothAssetEditorCore/Source/ChaosClothAssetEditorTools/Private/ChaosClothAsset/ClothWeightMapPaintTool.cpp
+++ b/Engine/Plugins/ChaosClothAssetEditorCore/Source/ChaosClothAssetEditorTools/Private/ChaosClothAsset/ClothWeightMapPaintTool.cpp
@@ -363,7 +363,7 @@ void UClothEditorWeightMapPaintTool::Setup()
 	const int32 NumExpectedWeights = bHaveDynamicMeshToWeightConversion ? WeightToDynamicMesh.Num() : GetSculptMesh()->MaxVertexID();
 	TArray<float> CurrentWeights;
 	CurrentWeights.SetNumZeroed(NumExpectedWeights);
-	WeightMapNodeToUpdate->CalculateFinalVertexWeightValues(InputWeightMap, TArrayView<float>(CurrentWeights));
+	WeightMapNodeToUpdate->CalculateFinalVertexWeightValues(nullptr, InputWeightMap, TArrayView<float>(CurrentWeights));
 
 	if (bHaveDynamicMeshToWeightConversion)
 	{
