ChaosDeformable: Add FChaosDeformableSimulationAssetCollisionConfigNode for per-asset collision config; fix FCreateSpringConstraintNode to follow triangle-point convention with mass-based stiffness; update SimulationBaseConfigNode and terminal node

#Synced-CL 52219429

#Synced-CL 52219429

diff --git a/Engine/Plugins/ChaosClothAssetDataflowNodes/Source/ChaosClothAssetDataflowNodes/Private/ChaosClothAsset/SimulationBaseConfigNode.cpp b/Engine/Plugins/ChaosClothAssetDataflowNodes/Source/ChaosClothAssetDataflowNodes/Private/ChaosClothAsset/SimulationBaseConfigNode.cpp
--- a/Engine/Plugins/ChaosClothAssetDataflowNodes/Source/ChaosClothAssetDataflowNodes/Private/ChaosClothAsset/SimulationBaseConfigNode.cpp
+++ b/Engine/Plugins/ChaosClothAssetDataflowNodes/Source/ChaosClothAssetDataflowNodes/Private/ChaosClothAsset/SimulationBaseConfigNode.cpp
@@ -543,18 +543,18 @@ void FChaosClothAssetSimulationBaseConfigNode::FPropertyHelper::SetSolverPropert
 	SetPropertyWeighted(PropertyName, PropertyValue, SimilarPropertyNames, PropertyFlags);
 }
 
-template void FChaosClothAssetSimulationBaseConfigNode::FPropertyHelper::SetSolverProperty<FChaosClothAssetImportedVectorValue>(const FName& PropertyName,
-	const FChaosClothAssetImportedVectorValue& PropertyValue, 
+template CHAOSCLOTHASSETDATAFLOWNODES_API void FChaosClothAssetSimulationBaseConfigNode::FPropertyHelper::SetSolverProperty<FChaosClothAssetImportedVectorValue>(const FName& PropertyName,
+	const FChaosClothAssetImportedVectorValue& PropertyValue,
 	const TFunction<FChaosClothAssetImportedVectorValue::ImportedType(UE::Chaos::ClothAsset::FCollectionClothFacade&)>& SolverValueFunction,
 	const TArray<FName>& SimilarPropertyNames, ECollectionPropertyFlags PropertyFlags);
 
-template void FChaosClothAssetSimulationBaseConfigNode::FPropertyHelper::SetSolverProperty<FChaosClothAssetImportedFloatValue>(const FName& PropertyName,
-	const FChaosClothAssetImportedFloatValue& PropertyValue, 
+template CHAOSCLOTHASSETDATAFLOWNODES_API void FChaosClothAssetSimulationBaseConfigNode::FPropertyHelper::SetSolverProperty<FChaosClothAssetImportedFloatValue>(const FName& PropertyName,
+	const FChaosClothAssetImportedFloatValue& PropertyValue,
 	const TFunction<FChaosClothAssetImportedFloatValue::ImportedType(UE::Chaos::ClothAsset::FCollectionClothFacade&)>& SolverValueFunction,
 	const TArray<FName>& SimilarPropertyNames, ECollectionPropertyFlags PropertyFlags);
 
-template void FChaosClothAssetSimulationBaseConfigNode::FPropertyHelper::SetSolverProperty<FChaosClothAssetImportedIntValue>(const FName& PropertyName,
-	const FChaosClothAssetImportedIntValue& PropertyValue, 
+template CHAOSCLOTHASSETDATAFLOWNODES_API void FChaosClothAssetSimulationBaseConfigNode::FPropertyHelper::SetSolverProperty<FChaosClothAssetImportedIntValue>(const FName& PropertyName,
+	const FChaosClothAssetImportedIntValue& PropertyValue,
 	const TFunction<FChaosClothAssetImportedIntValue::ImportedType(UE::Chaos::ClothAsset::FCollectionClothFacade&)>& SolverValueFunction,
 	const TArray<FName>& SimilarPropertyNames, ECollectionPropertyFlags PropertyFlags);
 	
diff --git a/Engine/Plugins/ChaosClothAssetDataflowNodes/Source/ChaosClothAssetDataflowNodes/Public/ChaosClothAsset/SimulationBaseConfigNode.h b/Engine/Plugins/ChaosClothAssetDataflowNodes/Source/ChaosClothAssetDataflowNodes/Public/ChaosClothAsset/SimulationBaseConfigNode.h
--- a/Engine/Plugins/ChaosClothAssetDataflowNodes/Source/ChaosClothAssetDataflowNodes/Public/ChaosClothAsset/SimulationBaseConfigNode.h
+++ b/Engine/Plugins/ChaosClothAssetDataflowNodes/Source/ChaosClothAssetDataflowNodes/Public/ChaosClothAsset/SimulationBaseConfigNode.h
@@ -33,7 +33,7 @@ public:
 
 	FChaosClothAssetSimulationBaseConfigNode() = default;
 
-	FChaosClothAssetSimulationBaseConfigNode(const UE::Dataflow::FNodeParameters& InParam, FGuid InGuid = FGuid::NewGuid());
+	CHAOSCLOTHASSETDATAFLOWNODES_API FChaosClothAssetSimulationBaseConfigNode(const UE::Dataflow::FNodeParameters& InParam, FGuid InGuid = FGuid::NewGuid());
 
 protected:
 	using ECollectionPropertyFlags = ::Chaos::Softs::ECollectionPropertyFlags;  // Short hand for property flag enums
@@ -52,7 +52,7 @@ protected:
 		template<typename T, typename PropertyType UE_REQUIRES(std::is_base_of_v<FChaosClothAssetSimulationBaseConfigNode, T> && ::Chaos::Softs::TIsWeightedType<PropertyType>::Value && !std::is_same_v<PropertyType, bool>)>
 		inline int32 SetProperty(const T* ConfigStruct, const PropertyType* PropertyValue, const TArray<FName>& SimilarPropertyNames = {}, ECollectionPropertyFlags PropertyFlags = ECollectionPropertyFlags::Animatable);
 
-		int32 SetPropertyBool(const FName& PropertyName, bool PropertyValue, const TArray<FName>& SimilarPropertyNames = {}, ECollectionPropertyFlags PropertyFlags = ECollectionPropertyFlags::Animatable);
+		CHAOSCLOTHASSETDATAFLOWNODES_API int32 SetPropertyBool(const FName& PropertyName, bool PropertyValue, const TArray<FName>& SimilarPropertyNames = {}, ECollectionPropertyFlags PropertyFlags = ECollectionPropertyFlags::Animatable);
 
 		template<typename T UE_REQUIRES(std::is_base_of_v<FChaosClothAssetSimulationBaseConfigNode, T>)>
 		inline int32 SetPropertyBool(const T* ConfigStruct, const bool* PropertyValue, const TArray<FName>& SimilarPropertyNames = {}, ECollectionPropertyFlags PropertyFlags = ECollectionPropertyFlags::Animatable);
@@ -93,7 +93,7 @@ protected:
 
 		/** Set an imported solver value onto the property */
 		template<typename PropertyType>
-		void SetSolverProperty(const FName& PropertyName, const PropertyType& PropertyValue, 
+		void SetSolverProperty(const FName& PropertyName, const PropertyType& PropertyValue,
 			const TFunction<typename PropertyType::ImportedType(UE::Chaos::ClothAsset::FCollectionClothFacade&)>& SolverValueFunction,
 			const TArray<FName>& SimilarPropertyNames, ECollectionPropertyFlags PropertyFlags = ECollectionPropertyFlags::None);
 
@@ -135,7 +135,7 @@ protected:
 		TSharedRef<FManagedArrayCollection> ClothCollection;
 	};
 
-	virtual void Evaluate(UE::Dataflow::FContext& Context, const FDataflowOutput* Out) const override;
+	CHAOSCLOTHASSETDATAFLOWNODES_API virtual void Evaluate(UE::Dataflow::FContext& Context, const FDataflowOutput* Out) const override;
 
 	UE_DEPRECATED(5.4, "Use AddProperties(FPropertyHelper&) instead.")
 	virtual void AddProperties(UE::Dataflow::FContext& Context, ::Chaos::Softs::FCollectionPropertyMutableFacade& Properties) const {}
@@ -146,7 +146,7 @@ protected:
 	/* Override this to do additional node-specific evaluate on the cloth collection output. AddProperties has already been called when this is called. */
 	virtual void EvaluateClothCollection(UE::Dataflow::FContext& Context, const TSharedRef<FManagedArrayCollection>& ClothCollection) const {}
 
-	void RegisterCollectionConnections();
+	CHAOSCLOTHASSETDATAFLOWNODES_API void RegisterCollectionConnections();
 
 	UE_DEPRECATED(5.4, "Use FPropertyHelper instead.")
 	int32 AddPropertyHelper(
@@ -160,7 +160,7 @@ protected:
 		return AddPropertyHelper(Properties, PropertyName, SimilarPropertyNames, Flags);
 	}
 
-	int32 AddPropertyHelper(
+	CHAOSCLOTHASSETDATAFLOWNODES_API int32 AddPropertyHelper(
 		::Chaos::Softs::FCollectionPropertyMutableFacade& Properties,
 		const FName& PropertyName,
 		const TArray<FName>& SimilarPropertyNames = {},
