Update tool tip for ClothCollisionThickness now that the weighted value is used by all collisions.
Clean up unused parameters and move rarely used ones to the advanced section.

#Synced-CL 52358950

diff --git a/Engine/Plugins/ChaosClothAssetDataflowNodes/Source/ChaosClothAssetDataflowNodes/Private/ChaosClothAsset/SimulationCollisionConfigNode.cpp b/Engine/Plugins/ChaosClothAssetDataflowNodes/Source/ChaosClothAssetDataflowNodes/Private/ChaosClothAsset/SimulationCollisionConfigNode.cpp
--- a/Engine/Plugins/ChaosClothAssetDataflowNodes/Source/ChaosClothAssetDataflowNodes/Private/ChaosClothAsset/SimulationCollisionConfigNode.cpp
+++ b/Engine/Plugins/ChaosClothAssetDataflowNodes/Source/ChaosClothAssetDataflowNodes/Private/ChaosClothAsset/SimulationCollisionConfigNode.cpp
@@ -26,8 +26,6 @@ void FChaosClothAssetSimulationCollisionConfigNode::AddProperties(FPropertyHelpe
 {
 	PropertyHelper.SetPropertyBool(this, &bUseCCD);
 	PropertyHelper.SetPropertyWeighted(this, &InnerCollisionThickness);
-
-	PropertyHelper.SetProperty(this, &ProximityStiffness);
 	
 	PropertyHelper.SetFabricProperty(FName(TEXT("CollisionThickness")), CollisionThicknessImported,
 		[](UE::Chaos::ClothAsset::FCollectionClothFabricFacade& FabricFacade)-> float
diff --git a/Engine/Plugins/ChaosClothAssetDataflowNodes/Source/ChaosClothAssetDataflowNodes/Public/ChaosClothAsset/SimulationCollisionConfigNode.h b/Engine/Plugins/ChaosClothAssetDataflowNodes/Source/ChaosClothAssetDataflowNodes/Public/ChaosClothAsset/SimulationCollisionConfigNode.h
--- a/Engine/Plugins/ChaosClothAssetDataflowNodes/Source/ChaosClothAssetDataflowNodes/Public/ChaosClothAsset/SimulationCollisionConfigNode.h
+++ b/Engine/Plugins/ChaosClothAssetDataflowNodes/Source/ChaosClothAssetDataflowNodes/Public/ChaosClothAsset/SimulationCollisionConfigNode.h
@@ -33,11 +33,11 @@ public:
 	bool bUsePlanarConstraintForSimpleColliders = false;
 
 	/** Enable colliding against any complex (e.g., SkinnedLevelSet, MLLevelSet) colliders.*/
-	UPROPERTY(EditAnywhere, Category = "Collision Properties", Meta = (InteractorName = "EnableComplexColliders"))
+	UPROPERTY(EditAnywhere, AdvancedDisplay, Category = "Collision Properties", Meta = (InteractorName = "EnableComplexColliders"))
 	bool bEnableComplexColliders = true;
 
 	/** Use Planar Constraints for complex (e.g., SkinnedLevelSet, MLLevelSet) colliders when doing multiple iterations. Planar constraints are cheaper than full collision detection, but less accurate. */
-	UPROPERTY(EditAnywhere, Category = "Collision Properties", Meta = (EditCondition = "bEnableComplexColliders", InteractorName = "UsePlanarConstraintForComplexColliders"))
+	UPROPERTY(EditAnywhere, AdvancedDisplay, Category = "Collision Properties", Meta = (EditCondition = "bEnableComplexColliders", InteractorName = "UsePlanarConstraintForComplexColliders"))
 	bool bUsePlanarConstraintForComplexColliders = true;
 
 	/** Enable colliding against any Skinned Triangle Mesh colliders.*/
@@ -48,13 +48,14 @@ public:
 	UPROPERTY(EditAnywhere, Category = "Collision Properties", Meta = (EditCondition = "bEnableSkinnedTriangleMeshCollisions", InteractorName = "UseSelfCollisionSubstepsForSkinnedTriangleMeshes"))
 	bool bUseSelfCollisionSubstepsForSkinnedTriangleMeshes = true;
 
-	/** Thickness added to the cloth when colliding against collision shapes. Currently only Skinned Triangle Meshes use the weighted value. All other collisions only use the Low value. */
+	/** Thickness added to the cloth when colliding against collision shapes. */
 	UPROPERTY(EditAnywhere, Category = "Collision Properties", DisplayName = "Cloth Collision Thickness", Meta = (InteractorName = "ClothCollisionThickness"))
 	FChaosClothAssetWeightedValue ClothCollisionThicknessImported = { true, UE::Chaos::ClothAsset::FDefaultFabric::ClothCollisionThickness,
 		UE::Chaos::ClothAsset::FDefaultFabric::ClothCollisionThickness, TEXT("ClothCollisionThickness"), true };
 
 	/** Stiffness for proximity repulsion forces (Force-based solver only). Units = kg cm/ s^2 (same as XPBD springs)*/
-	UPROPERTY(EditAnywhere, Category = "Proximity Force Properties", Meta = (UIMin = "0", UIMax = "10000", ClampMin = "0", ClampMax = "10000000", InteractorName = "ProximityStiffness"))
+	UE_DEPRECATED(5.8, "This property is not used.")
+	UPROPERTY()
 	float ProximityStiffness = 100.f;
 
 	/**
