Dataflow rendering : export method to create a section for rendering settings and add a section for cloth settings

#Synced-CL 52497132

diff --git a/Engine/Plugins/ChaosClothAssetDataflowNodes/Source/ChaosClothAssetDataflowNodes/Private/ChaosClothAsset/ClothRenderableTypes.cpp b/Engine/Plugins/ChaosClothAssetDataflowNodes/Source/ChaosClothAssetDataflowNodes/Private/ChaosClothAsset/ClothRenderableTypes.cpp
--- a/Engine/Plugins/ChaosClothAssetDataflowNodes/Source/ChaosClothAssetDataflowNodes/Private/ChaosClothAsset/ClothRenderableTypes.cpp
+++ b/Engine/Plugins/ChaosClothAssetDataflowNodes/Source/ChaosClothAssetDataflowNodes/Private/ChaosClothAsset/ClothRenderableTypes.cpp
@@ -26,6 +26,8 @@
 
 #include "Materials/Material.h"
 
+#define LOCTEXT_NAMESPACE "ClothRenderableTypes"
+
 namespace UE::Chaos::ClothAsset
 {
 	namespace Private
@@ -747,5 +749,16 @@ namespace UE::Chaos::ClothAsset
 		UE_DATAFLOW_REGISTER_RENDERABLE_TYPE(FClothSimSurfaceRenderableType);
 		UE_DATAFLOW_REGISTER_RENDERABLE_TYPE(FClothRenderSurfaceRenderableType);
 		UE_DATAFLOW_REGISTER_RENDERABLE_TYPE(FClothSimPatternRenderableType);
+
+
+		static const FName ClothCategoryName(TEXT("Cloth"));
+
+		UDataflowRenderableTypeSettings::RegisterSection(
+			UDataflowClothSimRenderSettings::StaticClass(),
+			"Cloth", LOCTEXT("Cloth", "Cloth"),
+			{ ClothCategoryName }
+		);
 	}
 }
+
+#undef LOCTEXT_NAMESPACE 
diff --git a/Engine/Plugins/Experimental/Dataflow/Source/DataflowEditor/Public/DataflowRendering/DataflowRenderableTypeSettings.h b/Engine/Plugins/Experimental/Dataflow/Source/DataflowEditor/Public/DataflowRendering/DataflowRenderableTypeSettings.h
--- a/Engine/Plugins/Experimental/Dataflow/Source/DataflowEditor/Public/DataflowRendering/DataflowRenderableTypeSettings.h
+++ b/Engine/Plugins/Experimental/Dataflow/Source/DataflowEditor/Public/DataflowRendering/DataflowRenderableTypeSettings.h
@@ -11,7 +11,7 @@ class UDataflowRenderableTypeSettings : public UObject
 {
 	GENERATED_BODY()
 public:
-	static void RegisterSection(UClass* Class, FName Name, const FText& DisplayName, TConstArrayView<FName> Categories);
+	DATAFLOWEDITOR_API static void RegisterSection(UClass* Class, FName Name, const FText& DisplayName, TConstArrayView<FName> Categories);
 };
 
 // Common types
