Dataflow : Add text debug info to the remesh node

#Synced-CL 52601091

diff --git a/Engine/Plugins/ChaosClothAssetDataflowNodes/Source/ChaosClothAssetDataflowNodes/Private/ChaosClothAsset/RemeshNode.cpp b/Engine/Plugins/ChaosClothAssetDataflowNodes/Source/ChaosClothAssetDataflowNodes/Private/ChaosClothAsset/RemeshNode.cpp
--- a/Engine/Plugins/ChaosClothAssetDataflowNodes/Source/ChaosClothAssetDataflowNodes/Private/ChaosClothAsset/RemeshNode.cpp
+++ b/Engine/Plugins/ChaosClothAssetDataflowNodes/Source/ChaosClothAssetDataflowNodes/Private/ChaosClothAsset/RemeshNode.cpp
@@ -32,6 +32,14 @@
 #include "Spatial/PointSetHashTable.h"
 #include "SubRegionRemesher.h"
 
+#if WITH_EDITOR
+#include "ChaosClothAsset/ClothDataflowViewModes.h"
+#include "Dataflow/DataflowDebugDrawInterface.h"
+#include "Dataflow/DataflowRenderingViewMode.h"
+#include "Misc/StringBuilder.h"
+#endif
+
+
 #include UE_INLINE_GENERATED_CPP_BY_NAME(RemeshNode)
 
 #define LOCTEXT_NAMESPACE "ChaosClothAssetRemeshNode"
@@ -2371,7 +2379,69 @@ void FChaosClothAssetRemeshNode_v2::Evaluate(UE::Dataflow::FContext& Context, co
 
 }
 
+#if WITH_EDITOR
 
+bool FChaosClothAssetRemeshNode_v2::CanDebugDraw() const
+{
+	return true;
+}
+bool FChaosClothAssetRemeshNode_v2::CanDebugDrawViewMode(const FName& ViewModeName) const
+{
+	return true;//ViewModeName == UE::Dataflow::FDataflowConstruction3DViewMode::Name || ViewModeName == UE::Chaos::ClothAsset::FCloth3DSimViewMode::Name;
+}
+
+void FChaosClothAssetRemeshNode_v2::DebugDraw(UE::Dataflow::FContext& Context, IDataflowDebugDrawInterface& DataflowRenderingInterface, const FDebugDrawParameters& DebugDrawParameters) const
+{
+	if (DebugDrawParameters.bNodeIsSelected || DebugDrawParameters.bNodeIsPinned)
+	{
+		using namespace UE::Chaos::ClothAsset;
+
+		if (bRemeshSim || bRemeshRender)
+		{
+			FManagedArrayCollection InCollection = GetValue<FManagedArrayCollection>(Context, &Collection);
+			const TSharedRef<FManagedArrayCollection> InClothCollection = MakeShared<FManagedArrayCollection>(MoveTemp(InCollection));
+			FCollectionClothFacade InClothFacade(InClothCollection);
+
+			TStringBuilder<1024> Builder;
+			if (InClothFacade.IsValid())  // Can only act on the collection if it is a valid cloth collection
+			{
+				Builder.Appendf(TEXT("Original:\n"));
+				if (bRemeshSim)
+				{
+					Builder.Appendf(TEXT("- Sim 3D mesh: %d vertices, %d triangles\n"), InClothFacade.GetNumSimVertices3D(), InClothFacade.GetNumSimFaces());
+				}
+				if (bRemeshRender)
+				{
+					Builder.Appendf(TEXT("- Render mesh: %d vertices, %d triangles\n"), InClothFacade.GetNumRenderVertices(), InClothFacade.GetNumRenderFaces());
+				}
+			}
+
+			if (const FDataflowOutput* Output = FindOutput(&Collection))
+			{
+				FManagedArrayCollection OutCollection = Output->ReadValue(Context, Collection);
+				const TSharedRef<FManagedArrayCollection> OutClothCollection = MakeShared<FManagedArrayCollection>(MoveTemp(OutCollection));
+				FCollectionClothFacade OutClothFacade(OutClothCollection);
+				if (OutClothFacade.IsValid())  // Can only act on the collection if it is a valid cloth collection
+				{
+					Builder.Appendf(TEXT("Remeshed:\n"));
+					if (bRemeshSim)
+					{
+						Builder.Appendf(TEXT("- Sim 3D mesh: %d vertices, %d triangles\n"), OutClothFacade.GetNumSimVertices3D(), OutClothFacade.GetNumSimFaces());
+					}
+					if (bRemeshRender)
+					{
+						Builder.Appendf(TEXT("- Render mesh: %d vertices, %d triangles\n"), OutClothFacade.GetNumRenderVertices(), OutClothFacade.GetNumRenderFaces());
+					}
+				}
+			}
+			if (Builder.Len() > 0)
+			{
+				DataflowRenderingInterface.DrawOverlayText(Builder.ToString());
+			}
+		}
+	}
+}
+#endif
 
 
 ///////////////////////////////////////////////////////////////////////////////////////////////////////////////////////////////
diff --git a/Engine/Plugins/ChaosClothAssetDataflowNodes/Source/ChaosClothAssetDataflowNodes/Public/ChaosClothAsset/RemeshNode.h b/Engine/Plugins/ChaosClothAssetDataflowNodes/Source/ChaosClothAssetDataflowNodes/Public/ChaosClothAsset/RemeshNode.h
--- a/Engine/Plugins/ChaosClothAssetDataflowNodes/Source/ChaosClothAssetDataflowNodes/Public/ChaosClothAsset/RemeshNode.h
+++ b/Engine/Plugins/ChaosClothAssetDataflowNodes/Source/ChaosClothAssetDataflowNodes/Public/ChaosClothAsset/RemeshNode.h
@@ -137,6 +137,12 @@ private:
 
 	virtual void Evaluate(UE::Dataflow::FContext & Context, const FDataflowOutput * Out) const override;
 
+#if WITH_EDITOR
+	virtual bool CanDebugDraw() const override;
+	virtual bool CanDebugDrawViewMode(const FName& ViewModeName) const override;
+	virtual void DebugDraw(UE::Dataflow::FContext& Context, IDataflowDebugDrawInterface& DataflowRenderingInterface, const FDebugDrawParameters& DebugDrawParameters) const override;
+#endif
+
 	void GetSimSelectionFromContext(UE::Dataflow::FContext& Context, const TSharedRef<FManagedArrayCollection>& Collection, const FString& SimDensityMap,  TSet<int32>& OutSelection) const;
 	void GetRenderSelectionFromContext(UE::Dataflow::FContext& Context, const TSharedRef<FManagedArrayCollection>& Collection, const FString& RenderDensityMap, TSet<int32>& OutSelection) const;
 };
